Fix linking of ScreenShotComparison on Mac - one of the functions in Slate depends on UObjectBase::GetFName().

diff --git a/Engine/Source/Developer/ScreenShotComparison/ScreenShotComparison.Build.cs b/Engine/Source/Developer/ScreenShotComparison/ScreenShotComparison.Build.cs
--- a/Engine/Source/Developer/ScreenShotComparison/ScreenShotComparison.Build.cs
+++ b/Engine/Source/Developer/ScreenShotComparison/ScreenShotComparison.Build.cs
@@ -20,6 +20,7 @@ public class ScreenShotComparison : ModuleRules
 				"Slate",
 				"SlateCore",
 				"ImageWrapper",
+				"CoreUObject",
 			}
 		);
 
