Fullsized screenshots no longer show up in a fixed widow, they now show up in a client sizable window ensuring that a screenshot of the resolution of the users monitor or larger dosn't leave them in a hard to close window state.

diff --git a/Engine/Source/Developer/ScreenShotComparison/Private/Widgets/SScreenShotItem.cpp b/Engine/Source/Developer/ScreenShotComparison/Private/Widgets/SScreenShotItem.cpp
--- a/Engine/Source/Developer/ScreenShotComparison/Private/Widgets/SScreenShotItem.cpp
+++ b/Engine/Source/Developer/ScreenShotComparison/Private/Widgets/SScreenShotItem.cpp
@@ -83,10 +83,10 @@ FReply SScreenShotItem::OnImageClicked(const FGeometry& InGeometry, const FPoint
 	{
 		FString AssetFilename = ScreenShotData->GetAssetName();
 
-		TSharedRef<SScreenShotImagePopup> PopupImage = 
+		TSharedRef<SScreenShotImagePopup> PopupImage =
 			SNew(SScreenShotImagePopup)
-			. ImageAssetName(FName(*AssetFilename))
-			. ImageSize(ImageSize);
+			.ImageAssetName(FName(*AssetFilename))
+			.ImageSize(ImageSize);
 
 		auto ParentWindow = FSlateApplication::Get().FindWidgetWindow(AsShared()).ToSharedRef();
 
@@ -94,10 +94,11 @@ FReply SScreenShotItem::OnImageClicked(const FGeometry& InGeometry, const FPoint
 		auto PopupWindow = SNew(SWindow)
 			.Title( FText::FromString(AssetFilename))
 			.IsPopupWindow(false)
-			.SizingRule(ESizingRule::Autosized)
+			.SizingRule(ESizingRule::UserSized)
+			.ClientSize(FVector2D(1024,768))
 			.AutoCenter(EAutoCenter::PreferredWorkArea)
-			.SupportsMaximize(false)
-			.SupportsMinimize(false)
+			.SupportsMaximize(true)
+			.SupportsMinimize(true)
 			.FocusWhenFirstShown(true)
 			.ActivateWhenFirstShown(true)
 			.Content()
