Copying //UE4/Dev-Editor to //UE4/Dev-Main (Source: //UE4/Dev-Editor @ 3228984)
==========================
MAJOR FEATURES + CHANGES
==========================

Change 3168749 on 2016/10/20 by [REDACTED_PERSON]

	Fixed bug in csgRebuild where dynamic brushes from the whole world are rebuilt instead of just those from the current level. csgRebuild is supposed to act only on the current level's geometry.

Change 3169740 on 2016/10/20 by [REDACTED_PERSON]

	Automation - Removing old screenshots, working on new naming convention.

Change 3169796 on 2016/10/20 by [REDACTED_PERSON]

	Automation - Adding new screenshots.

Change 3169800 on 2016/10/20 by [REDACTED_PERSON]

	Automation - Working on improvements to screenshot comparions, now using the Unique device id instead of adapter name.  Working on better metadata based matching for which screenshot to use, stubbing in support for adding alternative versions of screenshots.

Change 3169901 on 2016/10/20 by [REDACTED_PERSON]

	Automation - More fixes / refinements to the way we add alternatives and replace old versions of screenshots.

Change 3169926 on 2016/10/20 by [REDACTED_PERSON]

	Added extension point for level editor viewport's Show and Camera menus

Change 3170053 on 2016/10/20 by [REDACTED_PERSON]

	Back out changelist 3169926

Change 3170067 on 2016/10/20 by [REDACTED_PERSON]

	Added extension point for level editor viewport's Show and Camera menus

Change 3170382 on 2016/10/21 by [REDACTED_PERSON]

Change 3170520 on 2016/10/21 by [REDACTED_PERSON]

Change 3170522 on 2016/10/21 by [REDACTED_PERSON]

Change 3170524 on 2016/10/21 by [REDACTED_PERSON]

Change 3170530 on 2016/10/21 by [REDACTED_PERSON]

	Moved the HasKey method from UMapProperty to FScriptMapHelper, and moved the HasElement property from USetProperty to FScriptSetHelper

Change 3170768 on 2016/10/21 by [REDACTED_PERSON]

	Back out changelist 3170067

Change 3170795 on 2016/10/21 by [REDACTED_PERSON]

	JsonObjectConverter - By default UStructToJsonAttributes now skips transient properties.

Change 3170797 on 2016/10/21 by [REDACTED_PERSON]

	Automation - Fixing several warnings dealing with fbx testing.

Change 3170921 on 2016/10/21 by [REDACTED_PERSON]

	Automation - Fixing more warnings with FBX tests.

Change 3171109 on 2016/10/21 by [REDACTED_PERSON]

	Added extension point for level editor viewport Show menu

Change 3171812 on 2016/10/24 by [REDACTED_PERSON]

	Back out changelist 3163044

	This broke wrapping for Japanese and Chinese.

Change 3171842 on 2016/10/24 by [REDACTED_PERSON]

	Name each Parameter uniquely either from copy/paste of any creation menu
	Changed the default value for Scalar and Vector Parameter to 1 and 1,1,1,1
	Added a Promote To Parameter when clicking on an Input pin that will generate proper node type based on type pin type
	When editing a color property update the material expression preview

Change 3171958 on 2016/10/24 by [REDACTED_PERSON]

Change 3171969 on 2016/10/24 by [REDACTED_PERSON]

	Slate - Adding some code to prevent crashes if bogus user indexes are passed into SlateApplications GetUser functions.

Change 3171970 on 2016/10/24 by [REDACTED_PERSON]

	PR #2885: Fixed Stretched Landscape Editor Icons (Contributed by [REDACTED_PERSON])

Change 3172035 on 2016/10/24 by [REDACTED_PERSON]

	Fix to build warning for 3171970

Change 3172078 on 2016/10/24 by [REDACTED_PERSON]

Change 3172143 on 2016/10/24 by [REDACTED_PERSON]

	Line-break iterators will now avoid breaking words in Hangul

	The default behavior for wrapping Hangul is to use Western-style wrapping (where words are kept as-is) rather than East Asian-style (where words are broken by syllables). This behavior can be controlled by the Localization.HangulTextWrappingMethod CVar in-case you were dependant on the old behavior, but since modern Hangul uses spaces, the per-word wrapping is preferred by native speakers.

Change 3172418 on 2016/10/24 by [REDACTED_PERSON]

	Fixed Static Analysis error

Change 3173389 on 2016/10/25 by [REDACTED_PERSON]

Change 3173918 on 2016/10/25 by [REDACTED_PERSON]

Change 3173966 on 2016/10/25 by [REDACTED_PERSON]

Change 3174847 on 2016/10/26 by [REDACTED_PERSON]

Change 3174916 on 2016/10/26 by [REDACTED_PERSON]

	When re-importing fbx file, always log to the message log.

Change 3174940 on 2016/10/26 by [REDACTED_PERSON]

	Back out changelist 3174847 at request of platforms team. Was fixed on Main.

Change 3174995 on 2016/10/26 by [REDACTED_PERSON]

	Import commandlet fixes

	- Fixed crash when source control could not be contacted
	- Fixed assets not importing correctly if they depended on other assets in a previous import group within the automated import

Change 3175217 on 2016/10/26 by [REDACTED_PERSON]

	The FBX reimport animation code now return false if there was an error when importing

Change 3175728 on 2016/10/26 by [REDACTED_PERSON]

	Log a message when importing a skeletal mesh with more bone influence then the maximum supported
	#2875

Change 3177997 on 2016/10/28 by [REDACTED_PERSON]

	Editor - Prevent re-entrant calls when EndPlayMap is called.

Change 3178429 on 2016/10/28 by [REDACTED_PERSON]

	Engine - Bumping BaseEngine.ini to IOS_8, MinimumiOSVersion, as that is now the minimum allowed to fix an error on startup.  Tweaking the location of where some importing files go when they're imported.

Change 3179774 on 2016/10/31 by [REDACTED_PERSON]

	Guard against bad render targets in Slate RHI

Change 3179900 on 2016/10/31 by [REDACTED_PERSON]

	Added logging to track https://jira.it.epicgames.net/browse/UE-37900

Change 3179920 on 2016/10/31 by [REDACTED_PERSON]

	Removing LODs from skeletal meshes is now a transacted action and can be undone. Related to UE-25591.

Change 3179921 on 2016/10/31 by [REDACTED_PERSON]

Change 3180119 on 2016/10/31 by [REDACTED_PERSON]

	fbx importer avoid asset creation name clash

Change 3181905 on 2016/11/01 by [REDACTED_PERSON]

	Paint tool now allow users to paint on any vertex if they need it.

Change 3182355 on 2016/11/01 by [REDACTED_PERSON]

	We now support FBX LODs export for the asset exporter from the content browser.

Change 3183286 on 2016/11/02 by [REDACTED_PERSON]

	Make sure static mesh build settings are set properly when we re-import with different options. Specifically the normals, tangents and tangent space are dependent on the import options.

Change 3183567 on 2016/11/02 by [REDACTED_PERSON]

	The Content Browser's View Options originally included both Engine and GameProject plugins only when clicking Show Plugin Content. Since there are quite a few Engine plugins, this produces quite a bit of content in the Folders panel. Most of the Engine plugins have classes or content that isn't really meant to be user-facing, so the experience of hunting for a game plugin-in's content is poor.

	The new behavior is that GameProject plugins are controlled by the "View Plugin Content" option. In order to see the Engine plugins you'll need both Engine Content and Plugin Content checkboxes enabled. By default, the editor should enable the "View Plugin Content" checkbox since it should be limited to just the content in the game's Plugins folder.

Change 3184002 on 2016/11/02 by [REDACTED_PERSON]

	Fixed crash during TSF IME shutdown

Change 3185126 on 2016/11/03 by [REDACTED_PERSON]

	Some of the plugin templates define Editor specific plugins. If created and a Standalone build is run, the application will attempt to link in editor libraries in game mode and will run into issues when you hit any key. The fix is to specify an Editor module description for these plugins.
	Additionally, there appears to be a mismatch in pathing types when dealing with plugin path and GameDir. Plugin path is absolute and GameDir is relative by default. We check to see if the gameDir is a subset of the plugin path, but this fails due to the mismatch. The fix is to force both to be absolute (enforcing normalization of both paths as well).

Change 3185278 on 2016/11/03 by [REDACTED_PERSON]

	UMG - Fixing some issues with HDPI mode in the widget designer.

Change 3185355 on 2016/11/03 by [REDACTED_PERSON]

	UMG - Widget Component's Draw At Desired size now should also work correctly if it's in screenspace.

Change 3185510 on 2016/11/03 by [REDACTED_PERSON]

	UMG - Restoring the ability of the Widget Component to directly recieve hardware input.  The Widget Interaction Component is great for just about every interaction use case - the one it's not is when you actually want the 3D widgets to take focus, and to be able to be typed directly into by the user.  The kind of situation where you might want to use them as a 3D menu, in a non-VR environment.  By default - Widget Components will not behave in this manner, but you can now use the option bReceiveHardwareInput to enable the ability for Widget Components to function more like a widget in the screenspace of the viewport.

	Slate - The scene viewport now correctly takes scale into account when drawing the 'software cursor', this fixes an issue with HDPI mode, and the cursor not being restored to the same location after moving a gizmo.

Change 3185514 on 2016/11/03 by [REDACTED_PERSON]

	UMG - Fixing some HDPI mode problems with widget position calculation when projecting world to viewport / screen, absolute spaces.

Change 3185652 on 2016/11/03 by [REDACTED_PERSON]

	Slate - Exposing a cached version of the widget geometry that comes in during Tick.  Also performed a bit of optimization work on the class to make some space for the geometry object we now cache, by compacting the pointer event delegates we were storing.

Change 3185952 on 2016/11/03 by [REDACTED_PERSON]

	UMG - Fixing another build error relating to local widget geometry.

Change 3185953 on 2016/11/03 by [REDACTED_PERSON]

	UMG - Fixing a mac compiler warning.

Change 3186886 on 2016/11/04 by [REDACTED_PERSON]

	Fixed collapse all hiding everything in the settings editors

Change 3187014 on 2016/11/04 by [REDACTED_PERSON]

	Fixed new assets opening in a minimized window not restoring that window.

Change 3187026 on 2016/11/04 by [REDACTED_PERSON]

	UUnrealEdEngine::edactDeleteSelected calls out to FBlueprintEditorUtils::FindActorsThatReferenceActor. This checks the entire world for each actor to be deleted. When you have tens of thousands of actors  in the world and are deleting tens of thousands of actors, this can take minutes. This change amortizes the cost of finding the actor references once for the world and for each actor to be deleted, we query the cached list of references. This brings the deletion time down to seconds.

Change 3187073 on 2016/11/04 by [REDACTED_PERSON]

	Automation - Changing the code that writes out json to force no BOM as is the json standard.

Change 3187113 on 2016/11/04 by [REDACTED_PERSON]

	Removed double look-up in UTextProperty::SerializeItem

Change 3187114 on 2016/11/04 by [REDACTED_PERSON]

	Feedback context now uses culture correct percentage formatting

Change 3187273 on 2016/11/04 by [REDACTED_PERSON]

	Fbx importer for static mesh, make sure that we order the materials array to follow the section order.
	Add also some fbx automation test

Change 3187276 on 2016/11/04 by [REDACTED_PERSON]

	Fix crash when an actor picker shows up in the struct editor.  Structs do not have root property nodes

Change 3187463 on 2016/11/04 by [REDACTED_PERSON]

	Automation - Updating the blessed screenshots, and fixing the BOM issues with the json.

Change 3188638 on 2016/11/07 by [REDACTED_PERSON]

	Making the UI for adding/removing parameters in custom blueprint functions behave similarly to the struct creation dialog in the content browser. There are no longer "New" buttons at the bottom of the panel and the parameter moving controls have been moved onto the main parameter row instead of being nested inside the collapse panel. A tooltip will now let you know the full parameter name and type when you hover over the editable name field. Made the move up/down icons more legible by increasing contrast between the arrow and the light grey background.

Change 3189056 on 2016/11/07 by [REDACTED_PERSON]

	Core/Editor - UObject::IsAsset() now returns false if the outermost package is RF_Transient.  Also updating the creation of the transient package to be RF_Transient.  This makes it so transient packages created by UMG or some other editor for things like previewing a streamed in level instance, no longer show up in the content browser.

Change 3189147 on 2016/11/07 by [REDACTED_PERSON]

	Fixed potential race-condition where a UFont object could be GC'd while the loading screen was using the font cache

	This queues up the pending removal until it's safe to execute it (by a thread that fully owns Slate rendering).

Change 3189344 on 2016/11/07 by [REDACTED_PERSON]

	Merging //UE4/Dev-Main to Dev-Editor (//UE4/Dev-Editor)

Change 3189380 on 2016/11/07 by [REDACTED_PERSON]

	Guard against null object when creating details panel

Change 3190017 on 2016/11/08 by [REDACTED_PERSON]

	FrontX support for scene importer

Change 3190058 on 2016/11/08 by [REDACTED_PERSON]

	Fixed misaligned button in the new blueprint class dialog

Change 3190086 on 2016/11/08 by [REDACTED_PERSON]

	UMG - Fixing the calculation for widget componets screen position if the camera aspect is constrained.

Change 3190159 on 2016/11/08 by [REDACTED_PERSON]

	UMG - We no longer also take the platform DPI scale into account when applying UMG's UI scale.  Since UMG already provides a DPI scaling system, compounding it with the native OSes produces undesirable results, since the DPI scale curve does not take into account some unknown platform scale set by a user.

Change 3190161 on 2016/11/08 by [REDACTED_PERSON]

	UMG - UWidget is now Blueprintable.  Improving some doc.

Change 3190545 on 2016/11/08 by [REDACTED_PERSON]

	Support scaling when exporting skeleton (bind pose) to FBX

Change 3191614 on 2016/11/09 by [REDACTED_PERSON]

	Fix cooking crash after fbx import of a scene without meshes

Change 3191659 on 2016/11/09 by [REDACTED_PERSON]

	Cleanup Persona LOD section button layout

Change 3191882 on 2016/11/09 by [REDACTED_PERSON]

	Changed FBlackboardKeySelector::AddObjectFilter to use MakeUniqueObjectName so it generates more stable names, rather than relying on a static counter.

	Also updated FBlackboardKeySelector::AddClassFilter, FBlackboardKeySelector::AddEnumFilter, and FBlackboardKeySelector::AddNativeEnumFilter to use MakeUniqueObjectName to ensure they don't conflict.

Change 3192092 on 2016/11/09 by [REDACTED_PERSON]

	Deleting some test assets that were accidentally checked in, some of which no longer load

Change 3192281 on 2016/11/09 by [REDACTED_PERSON]

Change 3192365 on 2016/11/09 by [REDACTED_PERSON]

	Adding support for MaterialParameterCollections to Slate UI objects. For reasons of Blueprint controls amongst other things, MPC's are owned by individual UWorlds and transferred over to their respective Scenes. Since we want the latest values from those in-UWorld representations, Slate needs to know about the Scene on the render thread to properly map the materials to their MPC inputs. This involved keeping Scene arrays synchronized between the game logic thread and render thread, and adding a Scene index field to each batched draw element in Slate. SceneViewports are now responsible for registering their associated Scenes with the SlateRenderer. Since RetainerBoxes also draw their content as well, they too need to register their Scenes.

Change 3192494 on 2016/11/09 by [REDACTED_PERSON]

Change 3193183 on 2016/11/10 by [REDACTED_PERSON]

	When doing FBX scene re-import, the new staticmesh asset was not mark as dirty. So the system was not saving the new asset.

Change 3193419 on 2016/11/10 by [REDACTED_PERSON]

	Fixing UnrealTournament build error in SUTChatEditBox

Change 3193456 on 2016/11/10 by [REDACTED_PERSON]

	Fix to build warning C6011 in SWidgetHierarchyItem

Change 3193704 on 2016/11/10 by [REDACTED_PERSON]

	Create Cinematic Camera when importing camera from fbx

Change 3194593 on 2016/11/11 by [REDACTED_PERSON]

	Slate - Fixing the window reshaping logic to avoid work if we don't need to do it, rather than external calls attempting to do the check (poorly).  This appears to fix the problem with popup menus being slightly off in size, creating scrollbars.  This also prevents constant reshaping of windows, due to people performing the wrong checks over and over, because they were comparing against non-truncated or rounded values against truncated/rounded values.

Change 3194595 on 2016/11/11 by [REDACTED_PERSON]

	Slate - Simplifying the Menu Anchor popup code for new Windows, and correcting it so that it does not take non-DPI scale into account when calculating the size of the window.  Otherwise, popup menus on say, the blueprint editor change size depending upon the scale of the area.

Change 3194830 on 2016/11/11 by [REDACTED_PERSON]

	Optimized pasting brushes, so geometry is not constantly rebuilt for every brush that's added.  This improves performance by a couple of orders of magnitude!

Change 3194859 on 2016/11/11 by [REDACTED_PERSON]

	Fix fbx skeletal mesh cleanup material crash

Change 3195199 on 2016/11/11 by [REDACTED_PERSON]

	UMG - Updating the bindable widget searching code in sequencer to use the WidgetTree traversing code, instead of something custom.  This fixes the issue where it wasn't finding widgets inside of named slots.

Change 3196579 on 2016/11/14 by [REDACTED_PERSON]

	Guard against rendering crashes when a mesh with no lod resources is opened.

Change 3196614 on 2016/11/14 by [REDACTED_PERSON]

	Slate - The ignore incoming scale option for the scale box should now behave as expected in more cases.  It required modifying the GetRelativeLayoutScale function to also pass down the prepass scale, otherwise it can't extract out the incoming scale ahead of time before text is measured ahead of time.

Change 3196624 on 2016/11/14 by [REDACTED_PERSON]

	PR #2927: UE-38473: Shadow outline color uses shadow color (Contributed by [REDACTED_PERSON])

Change 3196770 on 2016/11/14 by [REDACTED_PERSON]

	Ensure instead of crash when updating the selection pivot if a component's actor is not selected (this is non fatal)

Change 3196863 on 2016/11/14 by [REDACTED_PERSON]

	Slate - Allowing font outline settings to be specified in native code when constructing a SlateFontInfo via a ctor.

Change 3196900 on 2016/11/14 by [REDACTED_PERSON]

	Slate - Upgrading some cases that were using the older version of GetRelativeLayoutScale.

Change 3196947 on 2016/11/14 by [REDACTED_PERSON]

	Guard against crashes in the details panel when an OS message causes the tree to refresh when a previous event has invalidate the contents of the details panel.

Change 3197028 on 2016/11/14 by [REDACTED_PERSON]

	Shift Drag is not moving the camera when the user is dragging the 3 axis in same time.

Change 3197167 on 2016/11/14 by [REDACTED_PERSON]

	Removed pivot updating code per frame for now.  It changes on selection so I cant see a reason why it is needed every frame

Change 3197227 on 2016/11/14 by [REDACTED_PERSON]

	UMG/Blueprint - Exposing a way to set the default schema a blueprint editor derivation uses.  Updating all widget blueprints to finally use the WidgetGraphSchema.

Change 3197239 on 2016/11/14 by [REDACTED_PERSON]

	UMG - Improving the ReceiveHardwareInput option to limit exposure of widgets to hit testing that did not register for it.

Change 3197538 on 2016/11/14 by [REDACTED_PERSON]

	UMG - Making some progress on converting the schema over on load, now appear to correctly be loading it in time to be able to perform node conversions to convert older nodes to newer nodes.  Required changing the UBlueprint interface to have a virtual for upgrading nodes, that could be overriden in WidgetBlueprint to make sure the schemas have all been updated, as Serialize is too early, and PostLoad is too late.

Change 3198211 on 2016/11/15 by [REDACTED_PERSON]

	Guard against reimport factories being deleted while in use

Change 3198589 on 2016/11/15 by [REDACTED_PERSON]

Change 3198783 on 2016/11/15 by [REDACTED_PERSON]

	The Widget Component's Allow Hardware Input should now correctly convert coordinates coming from a viewport scaled up by the OS DPI scaling code.

Change 3198933 on 2016/11/15 by [REDACTED_PERSON]

	Changing the package localization ID used by a package now marks the package as dirty

Change 3198942 on 2016/11/15 by [REDACTED_PERSON]

	Clearing the package localization ID used by a package now marks the package as dirty

Change 3200241 on 2016/11/16 by [REDACTED_PERSON]

	Now allowing users to customize the Class Browser/Picker to filter out developer folders as well as hide internal use classes via INI settings. A ViewOptions button has been added to allow users to choose whether or not these filters are enabled.
	By default, internal only classes will be hidden and you will be limited to your own developer folder.

	Example change to DefaultEngine.ini or BaseEngine.ini to hide some classes as internal use
	[/Script/ClassViewer.ClassViewerProjectSettings]
	+InternalOnlyPaths=(Path="/Engine/VREditor")
	+InternalOnlyClasses=/Script/VREditor.VREditorBaseUserWidget

	The InternalOnlyPaths example will hide any classes in the VREditor folder or subfolders.
	The InternalOnlyClasses example will hide any classes that derive from VREditorBaseUserWidget.

	Both can be edited by the project settings UI so no manual INI tweaking is required. Please go to Project Settings->Class Viewer->Class Visibility Management

Change 3200621 on 2016/11/16 by [REDACTED_PERSON]

	Adding missing change needed post merge from main

Change 3200968 on 2016/11/16 by [REDACTED_PERSON]

	Fixed localization gather including texts that were instanced or otherwise unchanged

	- It now uses the archetype when exporting to diff against the default property value, and will only gather text that has changed from the default.
	- UMG widgets that are instanced from another UMG asset now only gather overridden values, and skip all child instances.

Change 3201033 on 2016/11/16 by [REDACTED_PERSON]

	Fixed source control to properly notify when files need to be checked out if a blueprint node is dragged

Change 3201829 on 2016/11/17 by [REDACTED_PERSON]

	Fixing issue where GEngine is null in early game loading, potentially causing a crash.

Change 3201832 on 2016/11/17 by [REDACTED_PERSON]

	Fix build warning

Change 3201835 on 2016/11/17 by [REDACTED_PERSON]

	Slate - Making it so explictly focusing a slate user that does not yet exist, creates the slate user so that the state is properly maintained in prepartion for that user's arrival / input.

Change 3201947 on 2016/11/17 by [REDACTED_PERSON]

	Fix streaming pause rendering starting a movie if a movie was already playing

Change 3202089 on 2016/11/17 by [REDACTED_PERSON]

	Editor - When replacing references, code that was added in 2729702, was allowing redirectors to be created that then might be abandoned and not renamed later if there was a collision on object name.  There's no problem if two objects have the same name, as long as they have different paths (except for classes).  So now the code records object paths in a seperate set, and avoids reprocessing / and creating multiple redirectors for the same objects, instead of just using object name.

Change 3202139 on 2016/11/17 by [REDACTED_PERSON]

	Fix for adjusting text spacing when lines are removed from TextLayouts

Change 3202398 on 2016/11/17 by [REDACTED_PERSON]

	Updated UMG Sequencer to properly fire events once per loop

Change 3202591 on 2016/11/17 by [REDACTED_PERSON]

	Fixing coding standards violations.

Change 3202744 on 2016/11/17 by [REDACTED_PERSON]

	StaticMeshComponent's OverriddenLightMapRes current displays the value it was set to, even when the bOverrideLightMapRes is false. The behavior within  UStaticMeshComponent::GetLightMapResolution is to use the LightMapResolution on the StaticMesh member instead when bOverrideLightMapRes is false. The UI was adjusted to reflect the more accurate behavior.

Change 3203009 on 2016/11/17 by [REDACTED_PERSON]

	Backing out changelist 3170522 per request

Change 3204077 on 2016/11/18 by [REDACTED_PERSON]

	Automation - Updating several bits of the screenshot automation piece to work a bit better, show names if we have them, and show preview dialogs for images.

Change 3204086 on 2016/11/18 by [REDACTED_PERSON]

	Added FGCObjectScopeGuard and TStrongObjectPtr as a convenient way to keep a UObject alive without having to add it to the root-set

	Both use FGCObject internally to reference the object and keep it alive.

	FGCObjectScopeGuard is designed to be lean and used as a guard for an existing pointer, whereas TStrongObjectPtr is more "full-fat" and designed to be a replacement for a raw-pointer.

	You should prefer FGCObjectScopeGuard where possible. Also note that TStrongObjectPtr isn't supported by UHT/UPROPERTY as you should just use a raw-pointer in that case (it would do the same thing).

Change 3204189 on 2016/11/18 by [REDACTED_PERSON]

	Removing content from dev folder

Change 3204205 on 2016/11/18 by [REDACTED_PERSON]

	Fix for being unable to delete folders that still have sub-folders in the Content Browser

Change 3204270 on 2016/11/18 by [REDACTED_PERSON]

	Fix StaticMesh socket reimports
	- socket transforms are now updated correctly on reimport
	- deleted socket from source will be removed on reimport
	- fix SocketManager refresh after import

Change 3204283 on 2016/11/18 by [REDACTED_PERSON]

Change 3205757 on 2016/11/21 by [REDACTED_PERSON]

	PR #2923: Slate: Fixed bug where NumCharactersInGlyph was set incorrectly for TAB characters (Contributed by [REDACTED_PERSON])

Change 3205759 on 2016/11/21 by [REDACTED_PERSON]

	PR #2958: Handle legacy Windows exe icon location (Contributed by [REDACTED_PERSON])

Change 3205816 on 2016/11/21 by [REDACTED_PERSON]

	PR #2956: Add plane to basicshapes (Contributed by [REDACTED_PERSON])

Change 3205831 on 2016/11/21 by [REDACTED_PERSON]

	Speculative fix for UE-38492

	This guards against null objects being passed to FAssetDeleteModel, as well as objects that become null due to the GC that happens in FAssetDeleteModel.

Change 3205869 on 2016/11/21 by [REDACTED_PERSON]

Change 3205873 on 2016/11/21 by [REDACTED_PERSON]

	Should also address the issue mentioned in #jira UE-32965

Change 3205954 on 2016/11/21 by [REDACTED_PERSON]

	Reverting changes from 3202744 that allowed override properties to show up as real properties in the list. There are several detail panel customizations that don't deal with this properly and rather than force everyone to upgrade, we'll just modify the static mesh detail customization to do the work.

Change 3205965 on 2016/11/21 by [REDACTED_PERSON]

	PR #2942

Change 3207157 on 2016/11/22 by [REDACTED_PERSON]

	Added UnrealWatchdog tool, run by the Editor, to improve abnormal shutdown tracking.
	[UE-32952] - Watchdog - Show CRC when reporting abnormal shutdowns in internal builds

Change 3207344 on 2016/11/22 by [REDACTED_PERSON]

	Added UnrealWatchdog to the Binary Release

Change 3207396 on 2016/11/22 by [REDACTED_PERSON]

	Add UnrealWatchdog to UGS precompiled binaries for Odin and Orion.

Change 3207418 on 2016/11/22 by [REDACTED_PERSON]

	Redid blur changes from Paragon Dev-General

	Blur widget updates
	- Renamed to SBackgroundBlur/UBackgroundBlur
	- Split SBackgroundBlur out into its own file
	- Added bApplyAlphaToBlur - when true, the strength of the blur is modulated by the widget alpha
	- Updated BlurRadius to be TOptional, so we auto-calculate radius when it isn't set
	- Added a UBackgroundBlurSlot, but left it unattached so it can be done in dev-editor (and update based on the engine version)
	- Updated OrionBlurWidget to export dll symbols and set up default low quality fallback image

Change 3207443 on 2016/11/22 by [REDACTED_PERSON]

	Fix CIS error on Mac from my change CL 3207157

Change 3207702 on 2016/11/22 by [REDACTED_PERSON]

	Added missing files

Change 3207958 on 2016/11/22 by [REDACTED_PERSON]

	Guard against crash clearing scenes from the slate RHI renderer during movie loading code.

Change 3207962 on 2016/11/22 by [REDACTED_PERSON]

	Added a guard against the rendering thread timing out while on a breakpoint by checking if the debugger is present before performing the timeout check

Change 3208194 on 2016/11/22 by [REDACTED_PERSON]

	Actually call correct method of checking for a debugger

Change 3209139 on 2016/11/23 by [REDACTED_PERSON]

	Adding support for "Show Only Modified Properties" filter to DetailWidgetRow

Change 3209206 on 2016/11/23 by [REDACTED_PERSON]

	Moving folders now removes the old folder from disk if it's empty

	This had already been done for deleting folders, but moving them was missed.

Change 3209281 on 2016/11/23 by [REDACTED_PERSON]

	PR #2932: Fix crash while updating cursor highlight (Contributed by [REDACTED_PERSON])

Change 3210383 on 2016/11/25 by [REDACTED_PERSON]

	Documented Crash Report Client analytics events
	[UE-32787] - Document Crash Report Client analytics events in code

Change 3210385 on 2016/11/25 by [REDACTED_PERSON]

	Make sure the combine mesh option of the staticmesh import is stored in staticmeshimportdata so the re-import know if it must re-import in combined or not

Change 3210983 on 2016/11/28 by [REDACTED_PERSON]

	Merging //UE4/Dev-Main to Dev-Editor (//UE4/Dev-Editor)

Change 3211001 on 2016/11/28 by [REDACTED_PERSON]

	Fix build errors

Change 3211009 on 2016/11/28 by [REDACTED_PERSON]

	PR #2960: Git plugin: multiline initial commit message and other connect screen cleanup (Contributed by [REDACTED_PERSON])

Change 3211017 on 2016/11/28 by [REDACTED_PERSON]

	Fix ATSC texture compression quality tooltip

Change 3211045 on 2016/11/28 by [REDACTED_PERSON]

	Fix compile errors

Change 3211081 on 2016/11/28 by [REDACTED_PERSON]

	Fix post process anim blueprints on skeletal meshes not being clearable

Change 3211094 on 2016/11/28 by [REDACTED_PERSON]

	Added more logging for jira UE-39000

Change 3211284 on 2016/11/28 by [REDACTED_PERSON]

	Redid fix for UE-35822 in dev-editor

Change 3211544 on 2016/11/28 by [REDACTED_PERSON]

	Fix deprecation warning

Change 3211769 on 2016/11/28 by [REDACTED_PERSON]

	Disable motion blur in editor views by default

Change 3211776 on 2016/11/28 by [REDACTED_PERSON]

	Fix PS4 compile errors

Change 3211949 on 2016/11/28 by [REDACTED_PERSON]

	Details panels changes
	- Added the ability to create groups within groups in details panel customizations
	- Added the ability for struct customizations to add categories to the parent

Change 3211954 on 2016/11/28 by [REDACTED_PERSON]

	Reorganized the post process settings so they appear as categories in the parent and so that they have better categories to make it clear what all the settings do.

Change 3213158 on 2016/11/29 by [REDACTED_PERSON]

	Updated User Defined Enum display names to use real FText instances so they can have stable keys

	This avoids the issue where the FText display names were cached from an FString, resulting in them having a different identity each time they were re-cached, which lead to localization and deterministic cooking issues.

	User Defined Enums no longer use meta-data to store their display names, and instead use a TMap in UUserDefinedEnum to map the raw enum entry name to its friendly display name. In addition to this, the enum editor has been updated to use STextPropertyEditableTextBox, which will keep the keys used by the display names stable where possible (allowing for delta-localization and historic tracking).

Change 3213172 on 2016/11/29 by [REDACTED_PERSON]

	Adding experimental support for content hot-reloading

	The underlying support for this is in CoreUObject (see ReloadPackage and ReloadPackages in UObjectGlobals.h/.cpp), with editor specific support being added via PackageTools::ReloadPackages, and also hooks registered with FCoreUObjectDelegates::OnPackageReloaded (eg, UEditorEngine::HandlePackageReloaded).

	The basic workflow for package reloading is as follows:
	 1) The current package is renamed, and the RF_NewerVersionExists flag is added to it and all its sub-objects.
	 2) The new package is loaded. Should this fail the old package is renamed back, and the RF_NewerVersionExists flag is removed.
	 3) We generate a mapping between objects in the old package and objects in the new package (see UObject::BuildSubobjectMapping).
	 4) We enumerate every object in memory, and fix-up any serialized or ARO object pointers referencing something from the old package, to reference the equivalent object from the new package (or null if no object could be found).
	 5) We run a GC, and verify that the old package was purged (printing any lingering references if it wasn't).

	For efficiency reasons package reloading may be run in batches (the editor uses batches of 500), as this allows package reloading to happen faster (as the reference fix-up and GC only happens once per-batch) at the cost of consuming more memory.

	In-editor there is an experimental setting to enable content hot-reloading. When this is enabled the SCC operations in the Content Browser will use content hot-reloading, rather than attempt to unload the reload the package as separate operations (which often fails). In order to allow the external SCC program to overwrite the files on disk, the linkers are detached from any packages that will be replaced prior to invoking the internal SCC operation.

Change 3213428 on 2016/11/29 by [REDACTED_PERSON]

	Implemented clamping on FTextInputMethodContext::SetSelectionRange to fix an issue where composition could provide an invalid range if the text was changed while composing

Change 3213442 on 2016/11/29 by [REDACTED_PERSON]

	Workaround for a bug in TSF based MS IMEs on Windows 8+

	They omit calling GetSelection and instead expect QueryInsert to return the current selection range. This also seems to fix an issue where composition no longer worked once some text had been deleted.

Change 3213603 on 2016/11/29 by [REDACTED_PERSON]

	Changed PanelWidget::RemoveChildAt to not release slate resources if the child is a UserWidget

Change 3213633 on 2016/11/29 by [REDACTED_PERSON]

	Attempt to fix includetool cis warning

Change 3215159 on 2016/11/30 by [REDACTED_PERSON]

	Fixing MakeShared forward declaration

Change 3215220 on 2016/11/30 by [REDACTED_PERSON]

Change 3215390 on 2016/11/30 by [REDACTED_PERSON]

	Maps now end a hot-reload batch

Change 3215394 on 2016/11/30 by [REDACTED_PERSON]

	Updating guard to track down worlds that have no package as an outer

Change 3215500 on 2016/11/30 by [REDACTED_PERSON]

	Color grading widget customization

Change 3215519 on 2016/11/30 by [REDACTED_PERSON]

	Fixed crash caused by using TextNamespaceUtil::EnsurePackageNamespace in 'game' mode

Change 3215556 on 2016/11/30 by [REDACTED_PERSON]

	Fixed issue where check-out toast would not disappear

Change 3215585 on 2016/11/30 by [REDACTED_PERSON]

	Adding an explicit ESPMode to MakeShared to try and placate Android

Change 3215737 on 2016/11/30 by [REDACTED_PERSON]

	Fix build warning

Change 3215748 on 2016/11/30 by [REDACTED_PERSON]

	Guard against crashes due to duplicate items in the scene outliner if actors somehow end up attached to themselves

Change 3215758 on 2016/11/30 by [REDACTED_PERSON]

	Add a 'Custom...' build type for Dev-Editor.

Change 3216183 on 2016/11/30 by [REDACTED_PERSON]

	Fix win32 build error

Change 3216362 on 2016/11/30 by [REDACTED_PERSON]

	Fix mac build error.

Change 3216828 on 2016/12/01 by [REDACTED_PERSON]

	Fixing MakeShared on Android

Change 3216839 on 2016/12/01 by [REDACTED_PERSON]

	PR #2997: Spelling fix for Actor.h's description of bEnableAutoLODGeneration. (Contributed by [REDACTED_PERSON])

Change 3216842 on 2016/12/01 by [REDACTED_PERSON]

	Remove the ensure when pushing absolute transforms onto a canvas matrix stack.  We can handle this properly now by just adding the transform to the stack if the stack is empty

Change 3216874 on 2016/12/01 by [REDACTED_PERSON]

	Fix a number of keybindings problems
	- Removed editor keybindings from project settings.  It should not have been in there (already in editor settings)
	- Removed duplicate registration of editor keybindings from editor settings
	- Fixed memory leak regenerating keybinding widgets when ending PIE world.
	- Cleaned up styling a bit to make keybindings widgets clearer.

Change 3216881 on 2016/12/01 by [REDACTED_PERSON]

	Added support for reroute nodes to the material editor. These nodes should function identically to their counterparts in Blueprints. A new UMaterialExpression, UMaterialExpressionReroute has been added. It inserts no HLSL code, and instead just moves along its input to find the real UMaterialExpression that it is ultimately bound to.
	Since the material system serializes its data as UMaterialExpressions, a more generalized approach across graph types isn't available as only the visual UI layer is shared between blueprints and material graphs.

	Also modified the material palette and popup material expression menu to allow for c++ based material name and description customization. If we choose to expand this, it would make the C++ material nodes more discoverable and understandable.

	Manually pulled in CL 3200823 and 3208490 to get bugfixes around material attribute usage.

	Adding an reroute node should function identically to Blueprints (ie double-click on connection to add or Utility\Add Reroute Node from palette). You should be able to add as many reroute nodes as you want in a chain. A reroute node that only has a connected output and not an input should behave as if there were no reroute node present (i.e. triggering constants on Add). It should be possible to use reroute nodes between any two supported node types if they are connectable in isolation. Where possible, we should show the same type mismatch errors that you'd see if connecting nodes directly (ie dragging a boolean constant into a reroute node connected to an Add should result in a Float/Bool mismatch). A reroute node is purely visual, it should have no impact on the final instruction count.

	In the event that an incomplete reroute input was completed by dragging to an invalid type, I tried to guarantee that the compiler would generate the appropriate errors. This can happen because we only know the inputs to a given node in code. If a reroute node doesn't have an input, it does not know what type it should be. However, the compiler should still detect these bad cases and error out.

Change 3216968 on 2016/12/01 by [REDACTED_PERSON]

	Syncing via source control now unloads (rather than reloads) packages that have been deleted from disk

Change 3216970 on 2016/12/01 by [REDACTED_PERSON]

	Reverting files now uses hot-reloading (if enabled)

Change 3217233 on 2016/12/01 by [REDACTED_PERSON]

	You can now choose to reload dirty packages via content hot-reloading

	This will revert any in-memory changes to the asset, which may be useful when you want to roll it back to its initial state without restarting the editor.

Change 3217244 on 2016/12/01 by [REDACTED_PERSON]

	WindowsMoviePlayer:  Initialize the movie player texture on first frame regardless of whether or not the decoder has a sample ready.  This prevents a white texture from showing up for a frame.

Change 3217466 on 2016/12/01 by [REDACTED_PERSON]

	Fixed a bug where FTextFormatData::ConditionalCompile_NoLock would always compile the text even if it was up-to-date

Change 3217572 on 2016/12/01 by [REDACTED_PERSON]

	Using FText::Format with an invalid argument no longer strips any associated argument modifier data from the resultant formatted text

Change 3217688 on 2016/12/01 by [REDACTED_PERSON]

	Fixed crash reloading the active world package when it was dirty

Change 3217978 on 2016/12/01 by [REDACTED_PERSON]

	Fixed crash where the slate renderer holds into scenes during maps are loaded causing access to deleted data after the load is complete.  We clean up cached scenes each frame but if slate doesnt tick the scenes are not cleaned up.  This change moves the CleanupScenes code to a location that is called each tick and during map loads

Change 3218834 on 2016/12/02 by [REDACTED_PERSON]

	move some scene conversion import fbx options to staticmesh, skeletalmesh and animation import data so the re-import will have acces to those import options

Change 3218838 on 2016/12/02 by [REDACTED_PERSON]

	Fixed editing static mesh settings manually in the details panel not visually refreshing the collision primitives

Change 3218864 on 2016/12/02 by [REDACTED_PERSON]

	Fixed basic cube shape having a convex hull instead of a box for collision

Change 3218900 on 2016/12/02 by [REDACTED_PERSON]

	Move static mesh collision properties to the collision category

Change 3219143 on 2016/12/02 by [REDACTED_PERSON]

	Change the way GetRandomVectorInBrush generate the Start/end position to use the BrushNormal instead of the BrushDirection

Change 3219199 on 2016/12/02 by [REDACTED_PERSON]

	Fixed a crash when changing Physical Surface Name and reassigning it on a physical material that uses it

Change 3219358 on 2016/12/02 by [REDACTED_PERSON]

	Fix fbx automation tests

Change 3219362 on 2016/12/02 by [REDACTED_PERSON]

	Support for MAX multisub material

Change 3219774 on 2016/12/02 by [REDACTED_PERSON]

	PR #2888: Add a setting to allow the Sources Panel to expand by default (Contributed by [REDACTED_PERSON])

Change 3219793 on 2016/12/02 by [REDACTED_PERSON]

	SWindow now restores focus back to the widget that last had focus when it was deactivated

Change 3221272 on 2016/12/05 by [REDACTED_PERSON]

	UI background blur tweaks
	- Adjust the downsample amount for lower kernel sizes
	- Flush post process memory used by the blur when switching levels

Change 3221273 on 2016/12/05 by [REDACTED_PERSON]

	Added guards against accesing scene caching methods of the slate resource manager on the rendering thread

Change 3221392 on 2016/12/05 by [REDACTED_PERSON]

	Added basic support for playing safe movies very early in the engine startup sequence.  A movie is considered safe to play very early if it is just a movie file and not some complex slate based UI loading screen

	no platform actually supports this yet as none of the movie streamer modules are loaded early enough and many platforms cant render this early

	Set PLATFORM_SUPPORTS_EARLY_MOVIE_PLAYBACK to 1 for your platform if it supports early loading

Change 3221831 on 2016/12/05 by [REDACTED_PERSON]

	Fixed UNumericProperty::ReadEnumAsUint8 not considering enum redirects when resolving the name

Change 3221986 on 2016/12/05 by [REDACTED_PERSON]

	Added an "Inline" font loading method

	This can be used in a cooked build to store the font data within the Font Face asset itself (rather than a separate .ufont file) in order to guarantee a hitch free load, at the cost of potentially using more memory up-front.

	The existing "PreLoad" loading method has been renamed to "LazyLoad" to better reflect what it actually does.

	This also fixes a bug where FFontData::Serialize could try and use the referenced Font Face asset before it had been fully loaded.

Change 3222065 on 2016/12/05 by [REDACTED_PERSON]

	Added log warning to detect hitches when lazily loading fonts

Change 3222225 on 2016/12/05 by [REDACTED_PERSON]

	Fixing style-set typo

Change 3223169 on 2016/12/06 by [REDACTED_PERSON]

	Fix autosaving prompting to check out built data if the built data asset was dirty during autosave

Change 3223184 on 2016/12/06 by [REDACTED_PERSON]

	Support LOD group and combine mesh

Change 3223212 on 2016/12/06 by [REDACTED_PERSON]

Change 3223215 on 2016/12/06 by [REDACTED_PERSON]

Change 3223249 on 2016/12/06 by [REDACTED_PERSON]

	Added back in missing code that was lost in a merge

Change 3223271 on 2016/12/06 by [REDACTED_PERSON]

Change 3223275 on 2016/12/06 by [REDACTED_PERSON]

	Fixed a race condition in WEX where the loading screen would render an external UI window that was referencing deleted materials

Change 3223276 on 2016/12/06 by [REDACTED_PERSON]

	Staticmesh socket fbx import.

Change 3223363 on 2016/12/06 by [REDACTED_PERSON]

	Reimport must ask for missing file when re-importing a old asset that has no source files

Change 3223423 on 2016/12/06 by [REDACTED_PERSON]

	Added option to place canvas panel children in same layer using explicit ZOrder setting.
	[UETOOL-935] - Figure out a solution for canvas panel batching

Change 3223551 on 2016/12/06 by [REDACTED_PERSON]

	UI mesh paint optimization, the slider now do not destroy the paint geometry adapter if the painted LOD has not change

Change 3223844 on 2016/12/06 by [REDACTED_PERSON]

	Back out change to change the defaults on vector and scalar expressions because this affects existing expressions that have not overridden the default

Change 3223880 on 2016/12/06 by [REDACTED_PERSON]

	Update doc links for maps and sets

Change 3224746 on 2016/12/07 by [REDACTED_PERSON]

Change 3224826 on 2016/12/07 by [REDACTED_PERSON]

Change 3224827 on 2016/12/07 by [REDACTED_PERSON]

	Improve search for material match on fbx mesh import
	- Add option to specify material search locations on mesh import
	- On Import it will now perform a first match material search in the following order	(suppose we are importing into /Game/Content/Assets/Meshes/MyMesh)
	   - Using Local as a search location will provide same behavior as before	(search non recursively in /Game/Content/Assets/Meshes)
	   - If option is UnderParent or more, search recursively in destination folder	(search recursively in /Game/Content/Assets/Meshes)
	   - If option is UnderParent or more, then recursively from parent folder	(search recursively in /Game/Content/Assets)
	   - If option is UnderRoot or more, search recursively from root folder	(search recursively in /Game)
	   - If option is AllAssets, search in every asset folder		(Search recursively everywhere)

Change 3224989 on 2016/12/07 by [REDACTED_PERSON]

	Fixed black callstack text in CrashReportClient.
	[UE-38987] - CrashReportClient Callstack text is rendering Black

Change 3225142 on 2016/12/07 by [REDACTED_PERSON]

	Added collapsing methods when exporting text for translation

	You can now choose how to collapse your text for translation from three export modes:
	 - ELocalizedTextCollapseMode::IdenticalTextIdAndSource - Collapse texts with the same text identity (namespace + key) and source text (default 4.15+ behavior).
	 - ELocalizedTextCollapseMode::IdenticalPackageIdTextIdAndSource - Collapse texts with the same package ID, text identity (namespace + key), and source text (4.14 behavior).
	 - ELocalizedTextCollapseMode::IdenticalNamespaceAndSource - Collapse texts with the same namespace and source text (legacy pre-4.14 behavior).

	The new default allows you to re-use the same text identity in different packages without having to translate the same text multiple times, and you can also now opt to get back to the legacy pre-4.14 behavior of collapsing all identical texts within the same namespace (in case you were reliant on that behavior).

	You can change this setting via the Localization Dashboard, or add it to your gather configs as "LocalizedTextCollapseMode" (this needs to go into any configs that deal with exporting or importing PO files - the default if nothing is specified is "ELocalizedTextCollapseMode::IdenticalTextIdAndSource").

Change 3225509 on 2016/12/07 by [REDACTED_PERSON]

	Static analysis fix, false positive

Change 3225859 on 2016/12/07 by [REDACTED_PERSON]

	Fix broken physical surface details customization
	- Scrolling now works properly
	- Edit boxes dont change size while editing
	- properly checks out or makes file writable once an edit has been made

Change 3226840 on 2016/12/08 by [REDACTED_PERSON]

	Fixing a bug in FText formatting where it would ignore the rebuild and Rebuild as Source arguments for the format string itself

Change 3226940 on 2016/12/08 by [REDACTED_PERSON]

	Avoid changing the W value when playing with the color grading wheel.

Change 3227814 on 2016/12/08 by [REDACTED_PERSON]

	 Temp disable lazy load font warnings to prevent infinite  recursion crashes at startup

Change 3228010 on 2016/12/08 by [REDACTED_PERSON]

	Fix for iOS compiling

Change 3228597 on 2016/12/09 by [REDACTED_PERSON]

	Removed hard dependency between UFont and UFontFace during struct serialization as it doesn't work with the EDL

Change 3228607 on 2016/12/09 by [REDACTED_PERSON]

	Fixed infinite recursion caused by logging while the output log font was still being loaded

Change 3228770 on 2016/12/09 by [REDACTED_PERSON]

	Fixed UUserDefinedEnum::GetEnumText

	it was using GetNameByIndex (which includes C++ scoping), rather than GetEnumName (which doesn't). This was causing all name look-ups to fail.

Change 3228785 on 2016/12/09 by [REDACTED_PERSON]

	Fix static analysis warning

diff --git a/Engine/Build/InstalledEngineBuild.xml b/Engine/Build/InstalledEngineBuild.xml
--- a/Engine/Build/InstalledEngineBuild.xml
+++ b/Engine/Build/InstalledEngineBuild.xml
@@ -312,6 +312,7 @@
 		<Node Name="Build Tools Win64" Requires="Compile UnrealHeaderTool Win64">
 			<Compile Target="CrashReportClient" Configuration="Shipping" Platform="Win64" Tag="#Build Tools Win64"/>
 			<Compile Target="CrashReportClient" Configuration="Shipping" Platform="Win32" Tag="#Build Tools Win64"/>
+			<Compile Target="UnrealWatchdog" Configuration="Shipping" Platform="Win64" Tag="#Build Tools Win64"/>
 			<Compile Target="ShaderCompileWorker" Configuration="Development" Platform="Win64" Tag="#Build Tools Win64"/>
 			<Compile Target="UnrealCEFSubProcess" Configuration="Development" Platform="Win32" Tag="#Build Tools Win64"/>
 			<Compile Target="UnrealCEFSubProcess" Configuration="Development" Platform="Win64" Tag="#Build Tools Win64"/>
diff --git a/Engine/Config/BaseEditorPerProjectUserSettings.ini b/Engine/Config/BaseEditorPerProjectUserSettings.ini
--- a/Engine/Config/BaseEditorPerProjectUserSettings.ini
+++ b/Engine/Config/BaseEditorPerProjectUserSettings.ini
@@ -112,6 +112,11 @@ bShowOnlyActorsInCurrentLevel=False
 ; Scene Outliner Settings
 bShowOnlySelectedActors=False
 
+[/Script/UnrealEd.ClassViewerSettings]
+; Whether or not to display the hidden classes
+DisplayInternalClasses=False
+; How to filter the developer folder classes
+DeveloperFolderType=CVDT_CurrentUser
 
 [/Script/UnrealEd.ContentBrowserSettings]
 ; The number of objects to load at once in the Content Browser before displaying a warning about loading many assets
@@ -128,6 +133,8 @@ DisplayDevelopersFolder=False
 DisplayCollections=False
 ; Whether to display the C++ folders in the Content Browser
 DisplayCppFolders=True
+; Whether to display the game project plugins in the Content Browser
+DisplayPluginFolder=True
 
 [/Script/UnrealEd.DestructableMeshEditorSettings]
 ; Persona viewport default settings
diff --git a/Engine/Config/BaseEngine.ini b/Engine/Config/BaseEngine.ini
--- a/Engine/Config/BaseEngine.ini
+++ b/Engine/Config/BaseEngine.ini
@@ -709,6 +709,8 @@ bUseStreamingPause=false
 +TaggedPropertyRedirects=(ClassName="MultiLineEditableText",OldPropertyName="bAutoWrapText",NewPropertyName="AutoWrapText")
 +TaggedPropertyRedirects=(ClassName="MultiLineEditableTextBox",OldPropertyName="bAutoWrapText",NewPropertyName="AutoWrapText")
 
+; Remove "Widget" suffix from the BackgroundBlur widget (technically engine-level, but we're the only ones that used it before the rename)
++ActiveClassRedirects=(OldClassName="BackgroundBlurWidget",NewClassName="/Script/UMG.BackgroundBlur")
 ; Renamed ScriptComponent to ScriptPluginComponent
 +ActiveClassRedirects=(OldClassName="ScriptComponent",NewClassName="ScriptPluginComponent")
 
@@ -738,6 +740,8 @@ bUseStreamingPause=false
 
 +EnumRedirects=(EnumName="EComponentCreationMethod", OldEnumEntry="EComponentCreationMethod::ConstructionScript", NewEnumEntry="EComponentCreationMethod::SimpleConstructionScript")
 
++EnumRedirects=(EnumName="EFontLoadingPolicy",OldEnumEntry="EFontLoadingPolicy::PreLoad",NewEnumEntry="EFontLoadingPolicy::LazyLoad")
+
 +ActiveClassRedirects=(OldClassName="ESteamVRTrackingSpace",NewClassName="EHMDTrackingOrigin")
 +K2FieldRedirects=(OldFieldName="SteamVRFunctionLibrary.SetTrackingSpace", NewFieldName="HeadMountedDisplayFunctionLibrary.SetTrackingOrigin")
 +K2FieldRedirects=(OldFieldName="SteamVRFunctionLibrary.GetTrackingSpace", NewFieldName="HeadMountedDisplayFunctionLibrary.GetTrackingOrigin")
@@ -2376,7 +2380,7 @@ bForceSmokeTests=false
 
 [AutomationTesting.FbxImport]
 FbxImportTestPath=../../Content/FbxEditorAutomation/
-FbxImportTestPackagePath=/Engine/Content/FbxEditorAutomationOut
+FbxImportTestPackagePath=/Engine/FbxEditorAutomationOut
 
 [AutomationTesting.Blueprint]
 TestAllBlueprints=false
@@ -2831,5 +2835,11 @@ bStripAnimationDataOnDedicatedServer=False
 [/Script/Engine.MeshSimplificationSettings]
 r.MeshReductionModule="QuadricMeshReduction"
 
+[/Script/ClassViewer.ClassViewerProjectSettings]
++InternalOnlyPaths=(Path="/Engine/VREditor")
++InternalOnlyPaths=(Path="/Engine/Sequencer")
++InternalOnlyPaths=(Path="/Engine/NotForLicensees")
++InternalOnlyClasses=/Script/VREditor.VREditorBaseUserWidget
++InternalOnlyClasses=/Script/LevelSequence.LevelSequenceBurnIn
 [EngineSessionManager]
 UseWatchdogMTBF=false
\ No newline at end of file
diff --git a/Engine/Plugins/Editor/PluginBrowser/Source/PluginBrowser/Private/DefaultPluginWizardDefinition.cpp b/Engine/Plugins/Editor/PluginBrowser/Source/PluginBrowser/Private/DefaultPluginWizardDefinition.cpp
--- a/Engine/Plugins/Editor/PluginBrowser/Source/PluginBrowser/Private/DefaultPluginWizardDefinition.cpp
+++ b/Engine/Plugins/Editor/PluginBrowser/Source/PluginBrowser/Private/DefaultPluginWizardDefinition.cpp
@@ -19,8 +19,8 @@ void FDefaultPluginWizardDefinition::PopulateTemplatesSource()
 {
 	const FText BlankTemplateName = LOCTEXT("BlankLabel", "Blank");
 	const FText ContentOnlyTemplateName = LOCTEXT("ContentOnlyLabel", "Content Only");
-	const FText BasicTemplateName = LOCTEXT("BasicTemplateTabLabel", "Toolbar Button");
-	const FText AdvancedTemplateName = LOCTEXT("AdvancedTemplateTabLabel", "Standalone Window");
+	const FText BasicTemplateName = LOCTEXT("BasicTemplateTabLabel", "Editor Toolbar Button");
+	const FText AdvancedTemplateName = LOCTEXT("AdvancedTemplateTabLabel", "Editor Standalone Window");
 	const FText BlueprintLibTemplateName = LOCTEXT("BlueprintLibTemplateLabel", "Blueprint Library");
 	const FText EditorModeTemplateName = LOCTEXT("EditorModeTemplateLabel", "Editor Mode");
 	const FText ThirdPartyTemplateName = LOCTEXT("ThirdPartyTemplateLabel", "Third Party Library");
@@ -33,16 +33,16 @@ void FDefaultPluginWizardDefinition::PopulateTemplatesSource()
 	const FText EditorModeDescription = LOCTEXT("EditorModeDesc", "Create a plugin that will have an editor mode.\n\nThis will include a toolkit example to specify UI that will appear in \"Modes\" tab (next to Foliage, Landscape etc).\nIt will also include very basic UI that demonstrates editor interaction and undo/redo functions usage.");
 	const FText ThirdPartyDescription = LOCTEXT("ThirdPartyDesc", "Create a plugin that uses an included third party library.\n\nThis can be used as an example of how to include, load and use a third party library yourself.");
 
-	TemplateDefinitions.Add(MakeShareable(new FPluginTemplateDescription(ContentOnlyTemplateName, ContentOnlyDescription, TEXT("ContentOnly"), true)));
+	TemplateDefinitions.Add(MakeShareable(new FPluginTemplateDescription(ContentOnlyTemplateName, ContentOnlyDescription, TEXT("ContentOnly"), true, EHostType::Runtime)));
 	if (!bIsContentOnlyProject)
 	{
 		// Insert the blank template to make sure it appears before the content only template.
-		TemplateDefinitions.Insert(MakeShareable(new FPluginTemplateDescription(BlankTemplateName, BlankDescription, TEXT("Blank"))), 0);
-		TemplateDefinitions.Add(MakeShareable(new FPluginTemplateDescription(BasicTemplateName, BasicDescription, TEXT("Basic"))));
-		TemplateDefinitions.Add(MakeShareable(new FPluginTemplateDescription(AdvancedTemplateName, AdvancedDescription, TEXT("Advanced"))));
-		TemplateDefinitions.Add(MakeShareable(new FPluginTemplateDescription(BlueprintLibTemplateName, BlueprintLibDescription, TEXT("BlueprintLibrary"))));
-		TemplateDefinitions.Add(MakeShareable(new FPluginTemplateDescription(EditorModeTemplateName, EditorModeDescription, TEXT("EditorMode"))));
-		TemplateDefinitions.Add(MakeShareable(new FPluginTemplateDescription(ThirdPartyTemplateName, ThirdPartyDescription, TEXT("ThirdPartyLibrary"))));
+		TemplateDefinitions.Insert(MakeShareable(new FPluginTemplateDescription(BlankTemplateName, BlankDescription, TEXT("Blank"), true, EHostType::Developer)), 0);
+		TemplateDefinitions.Add(MakeShareable(new FPluginTemplateDescription(BlueprintLibTemplateName, BlueprintLibDescription, TEXT("BlueprintLibrary"), true, EHostType::Developer)));
+		TemplateDefinitions.Add(MakeShareable(new FPluginTemplateDescription(BasicTemplateName, BasicDescription, TEXT("Basic"), false, EHostType::Editor)));
+		TemplateDefinitions.Add(MakeShareable(new FPluginTemplateDescription(AdvancedTemplateName, AdvancedDescription, TEXT("Advanced"), false, EHostType::Editor)));
+		TemplateDefinitions.Add(MakeShareable(new FPluginTemplateDescription(EditorModeTemplateName, EditorModeDescription, TEXT("EditorMode"), false, EHostType::Editor)));
+		TemplateDefinitions.Add(MakeShareable(new FPluginTemplateDescription(ThirdPartyTemplateName, ThirdPartyDescription, TEXT("ThirdPartyLibrary"), true, EHostType::Developer)));
 	}
 }
 
diff --git a/Engine/Plugins/Editor/PluginBrowser/Source/PluginBrowser/Private/SNewPluginWizard.cpp b/Engine/Plugins/Editor/PluginBrowser/Source/PluginBrowser/Private/SNewPluginWizard.cpp
--- a/Engine/Plugins/Editor/PluginBrowser/Source/PluginBrowser/Private/SNewPluginWizard.cpp
+++ b/Engine/Plugins/Editor/PluginBrowser/Source/PluginBrowser/Private/SNewPluginWizard.cpp
@@ -512,10 +512,24 @@ FReply SNewPluginWizard::OnCreatePluginClicked()
 	{
 		bHasModules = true;
 	}
+	
+	EHostType::Type ModuleDescriptorType;
+
+	TArray<TSharedPtr<FPluginTemplateDescription>> SelectedTemplates = PluginWizardDefinition->GetSelectedTemplates();
+	// @todo For now assume there is one module type and all templates share it
+	if(SelectedTemplates.Num() > 0 )
+	{
+		ModuleDescriptorType = SelectedTemplates[0]->ModuleDescriptorType;
+	}
+	else
+	{
+		// Default to runtime
+		ModuleDescriptorType = EHostType::Runtime;
+	}
 
 	// Save descriptor file as .uplugin file
 	const FString UPluginFilePath = GetPluginFilenameWithPath();
-	bSucceeded = bSucceeded && WritePluginDescriptor(AutoPluginName, UPluginFilePath, PluginWizardDefinition->CanContainContent(), bHasModules);
+	bSucceeded = bSucceeded && WritePluginDescriptor(AutoPluginName, UPluginFilePath, PluginWizardDefinition->CanContainContent(), bHasModules, ModuleDescriptorType);
 
 	// Main plugin dir
 	const FString BasePluginFolder = GetPluginDestinationPath().ToString();
@@ -603,7 +617,7 @@ bool SNewPluginWizard::CopyFile(const FString& DestinationFile, const FString& S
 	}
 }
 
-bool SNewPluginWizard::WritePluginDescriptor(const FString& PluginModuleName, const FString& UPluginFilePath, bool bCanContainContent, bool bHasModules)
+bool SNewPluginWizard::WritePluginDescriptor(const FString& PluginModuleName, const FString& UPluginFilePath, bool bCanContainContent, bool bHasModules, EHostType::Type ModuleDescriptorType)
 {
 	FPluginDescriptor Descriptor;
 
@@ -614,7 +628,7 @@ bool SNewPluginWizard::WritePluginDescriptor(const FString& PluginModuleName, co
 
 	if (bHasModules)
 	{
-		Descriptor.Modules.Add(FModuleDescriptor(*PluginModuleName, EHostType::Developer));
+		Descriptor.Modules.Add(FModuleDescriptor(*PluginModuleName, ModuleDescriptorType));
 	}
 	Descriptor.bCanContainContent = bCanContainContent;
 
diff --git a/Engine/Plugins/Editor/PluginBrowser/Source/PluginBrowser/Private/SNewPluginWizard.h b/Engine/Plugins/Editor/PluginBrowser/Source/PluginBrowser/Private/SNewPluginWizard.h
--- a/Engine/Plugins/Editor/PluginBrowser/Source/PluginBrowser/Private/SNewPluginWizard.h
+++ b/Engine/Plugins/Editor/PluginBrowser/Source/PluginBrowser/Private/SNewPluginWizard.h
@@ -9,12 +9,14 @@
 #include "Widgets/DeclarativeSyntaxSupport.h"
 #include "Widgets/SCompoundWidget.h"
 #include "IPluginWizardDefinition.h"
+#include "ModuleDescriptor.h"
 
 class ITableRow;
 class SFilePathBlock;
 class STableViewBase;
 enum class ECheckBoxState : uint8;
 
+
 DECLARE_LOG_CATEGORY_EXTERN(LogPluginWizard, Log, All);
 
 class SFilePathBlock;
@@ -115,9 +117,12 @@ private:
 	 * Writes a plugin descriptor file to disk
 	 * @param PluginModuleName Name of the plugin and its module
 	 * @param UPluginFilePath Path where the descriptor file should be written
+	 * @param bCanContainContent Can this plugin contain content
+	 * @param bHasModules Does this plugin have Source?
+	 * @param ModuleDescriptorType If this plugin has Source, what is the type of Source included (so it can potentially be excluded in the right builds)
 	 * @return Whether the files was written successfully
 	 */
-	bool WritePluginDescriptor(const FString& PluginModuleName, const FString& UPluginFilePath, bool bCanContainContent, bool bHasModules);
+	bool WritePluginDescriptor(const FString& PluginModuleName, const FString& UPluginFilePath, bool bCanContainContent, bool bHasModules, EHostType::Type InModuleDescriptorType);
 
 	/**
 	 * Displays an editor pop up error notification
diff --git a/Engine/Plugins/Editor/PluginBrowser/Source/PluginBrowser/Public/IPluginWizardDefinition.h b/Engine/Plugins/Editor/PluginBrowser/Source/PluginBrowser/Public/IPluginWizardDefinition.h
--- a/Engine/Plugins/Editor/PluginBrowser/Source/PluginBrowser/Public/IPluginWizardDefinition.h
+++ b/Engine/Plugins/Editor/PluginBrowser/Source/PluginBrowser/Public/IPluginWizardDefinition.h
@@ -2,6 +2,8 @@
 
 #pragma once
 
+#include "ModuleDescriptor.h"
+
 /**
  * Description of a plugin template
  */
@@ -22,12 +24,16 @@ struct FPluginTemplateDescription
 	/** Can the plugin contain content? */
 	bool bCanContainContent;
 
+	/** What is the expected ModuleDescriptor type for this plugin?*/
+	EHostType::Type ModuleDescriptorType;
+
 	/** Constructor */
-	FPluginTemplateDescription(FText InName, FText InDescription, FString InOnDiskPath, bool bInCanContainContent = false)
+	FPluginTemplateDescription(FText InName, FText InDescription, FString InOnDiskPath, bool InCanContainContent, EHostType::Type InModuleDescriptorType)
 		: Name(InName)
 		, Description(InDescription)
 		, OnDiskPath(InOnDiskPath)
-		, bCanContainContent(bInCanContainContent)
+		, bCanContainContent(InCanContainContent)
+		, ModuleDescriptorType(InModuleDescriptorType)
 	{
 	}
 };
diff --git a/Engine/Plugins/Editor/SpeedTreeImporter/Source/SpeedTreeImporter/Classes/ReimportSpeedTreeFactory.h b/Engine/Plugins/Editor/SpeedTreeImporter/Source/SpeedTreeImporter/Classes/ReimportSpeedTreeFactory.h
--- a/Engine/Plugins/Editor/SpeedTreeImporter/Source/SpeedTreeImporter/Classes/ReimportSpeedTreeFactory.h
+++ b/Engine/Plugins/Editor/SpeedTreeImporter/Source/SpeedTreeImporter/Classes/ReimportSpeedTreeFactory.h
@@ -22,6 +22,7 @@ class UReimportSpeedTreeFactory : public USpeedTreeImportFactory, public FReimpo
 	virtual void SetReimportPaths(UObject* Obj, const TArray<FString>& NewReimportPaths) override;
 	virtual EReimportResult::Type Reimport(UObject* Obj) override;
 	virtual int32 GetPriority() const override;
+	virtual const UObject* GetFactoryObject() const { return this; }
 	//~ End FReimportHandler interface
 };
 
diff --git a/Engine/Plugins/Runtime/PhysXVehicles/Source/PhysXVehiclesEditor/Private/WheeledVehicleMovementComponent4WDetails.cpp b/Engine/Plugins/Runtime/PhysXVehicles/Source/PhysXVehiclesEditor/Private/WheeledVehicleMovementComponent4WDetails.cpp
--- a/Engine/Plugins/Runtime/PhysXVehicles/Source/PhysXVehiclesEditor/Private/WheeledVehicleMovementComponent4WDetails.cpp
+++ b/Engine/Plugins/Runtime/PhysXVehicles/Source/PhysXVehiclesEditor/Private/WheeledVehicleMovementComponent4WDetails.cpp
@@ -119,6 +119,17 @@ void FWheeledVehicleMovementComponent4WDetails::FSteeringCurveEditor::ModifyOwne
 	}
 }
 
+TArray<const UObject*> FWheeledVehicleMovementComponent4WDetails::FSteeringCurveEditor::GetOwners() const
+{
+	TArray<const UObject*> Owners;
+	if (Owner)
+	{
+		Owners.Add(Owner);
+	}
+
+	return Owners;
+}
+
 void FWheeledVehicleMovementComponent4WDetails::FSteeringCurveEditor::MakeTransactional()
 {
 	if (Owner)
@@ -163,6 +174,17 @@ void FWheeledVehicleMovementComponent4WDetails::FTorqueCurveEditor::ModifyOwner(
 	}
 }
 
+TArray<const UObject*> FWheeledVehicleMovementComponent4WDetails::FTorqueCurveEditor::GetOwners() const
+{
+	TArray<const UObject*> Owners;
+	if (Owner)
+	{
+		Owners.Add(Owner);
+	}
+
+	return Owners;
+}
+
 void FWheeledVehicleMovementComponent4WDetails::FTorqueCurveEditor::MakeTransactional()
 {
 	if (Owner)
diff --git a/Engine/Plugins/Runtime/PhysXVehicles/Source/PhysXVehiclesEditor/Private/WheeledVehicleMovementComponent4WDetails.h b/Engine/Plugins/Runtime/PhysXVehicles/Source/PhysXVehiclesEditor/Private/WheeledVehicleMovementComponent4WDetails.h
--- a/Engine/Plugins/Runtime/PhysXVehicles/Source/PhysXVehiclesEditor/Private/WheeledVehicleMovementComponent4WDetails.h
+++ b/Engine/Plugins/Runtime/PhysXVehicles/Source/PhysXVehiclesEditor/Private/WheeledVehicleMovementComponent4WDetails.h
@@ -28,6 +28,7 @@ private:
 		virtual TArray<FRichCurveEditInfoConst> GetCurves() const override;
 		virtual TArray<FRichCurveEditInfo> GetCurves() override;
 		virtual void ModifyOwner() override;
+		virtual TArray<const UObject*> GetOwners() const override;
 		virtual void MakeTransactional() override;
 		virtual void OnCurveChanged( const TArray<FRichCurveEditInfo>& ChangedCurveEditInfos ) override { }
 		virtual bool IsValidCurve(FRichCurveEditInfo CurveInfo) override;
@@ -45,6 +46,7 @@ private:
 		virtual TArray<FRichCurveEditInfoConst> GetCurves() const override;
 		virtual TArray<FRichCurveEditInfo> GetCurves() override;
 		virtual void ModifyOwner() override;
+		virtual TArray<const UObject*> GetOwners() const override;
 		virtual void MakeTransactional() override;
 		virtual void OnCurveChanged( const TArray<FRichCurveEditInfo>& ChangedCurveEditInfos ) override { }
 		virtual bool IsValidCurve( FRichCurveEditInfo CurveInfo ) override;
diff --git a/Engine/Plugins/Runtime/WindowsMoviePlayer/Source/WindowsMoviePlayer/Private/WindowsMovieStreamer.cpp b/Engine/Plugins/Runtime/WindowsMoviePlayer/Source/WindowsMoviePlayer/Private/WindowsMovieStreamer.cpp
--- a/Engine/Plugins/Runtime/WindowsMoviePlayer/Source/WindowsMoviePlayer/Private/WindowsMovieStreamer.cpp
+++ b/Engine/Plugins/Runtime/WindowsMoviePlayer/Source/WindowsMoviePlayer/Private/WindowsMovieStreamer.cpp
@@ -66,16 +66,15 @@ void FMediaFoundationMovieStreamer::ForceCompletion()
 bool FMediaFoundationMovieStreamer::Tick(float DeltaTime)
 {
 	FSlateTexture2DRHIRef* CurrentTexture = Texture.Get();
+	check(IsInRenderingThread());
 
-	if ((CurrentTexture != nullptr) && SampleGrabberCallback->GetIsSampleReadyToUpdate())
+	if (CurrentTexture && !CurrentTexture->IsInitialized())
 	{
-		check( IsInRenderingThread() );
-
-		if( !CurrentTexture->IsInitialized() )
-		{
-			CurrentTexture->InitResource();
-		}
+		CurrentTexture->InitResource();
+	}
 
+	if (CurrentTexture && SampleGrabberCallback->GetIsSampleReadyToUpdate())
+	{
 		uint32 Stride;
 		uint8* DestTextureData = (uint8*)RHILockTexture2D( CurrentTexture->GetTypedResource(), 0, RLM_WriteOnly, Stride, false );
 		FMemory::Memcpy( DestTextureData, TextureData.GetData(), TextureData.Num() );
diff --git a/Engine/Plugins/Tests/EditorTests/Source/EditorTests/Private/UnrealEd/EditorAssetAutomationTests.cpp b/Engine/Plugins/Tests/EditorTests/Source/EditorTests/Private/UnrealEd/EditorAssetAutomationTests.cpp
--- a/Engine/Plugins/Tests/EditorTests/Source/EditorTests/Private/UnrealEd/EditorAssetAutomationTests.cpp
+++ b/Engine/Plugins/Tests/EditorTests/Source/EditorTests/Private/UnrealEd/EditorAssetAutomationTests.cpp
@@ -624,6 +624,7 @@ bool FAssetEditorTest::RunTest(const FString& Parameters)
 	ASSET_TEST_CREATE(UBlueprint, UBlueprintMacroFactory, MPL, FactoryInst->ParentClass = AActor::StaticClass();)
 	ASSET_TEST_CREATE(UCameraAnim, UCameraAnimFactory, CA, )
 	ASSET_TEST_CREATE(UCurveBase, UCurveFactory, C, FactoryInst->CurveClass = UCurveFloat::StaticClass();)
+	FModuleManager::Get().LoadModule(TEXT("GameplayAbilities"));
 	UClass* GameplayAbilityClass = StaticLoadClass(UObject::StaticClass(), NULL, TEXT("GameplayAbilities.GameplayAbilitySet"), NULL, LOAD_None, NULL);
 	if (GameplayAbilityClass)
 	{
diff --git a/Engine/Source/Developer/AssetTools/Private/AssetTypeActions/AssetTypeActions_World.cpp b/Engine/Source/Developer/AssetTools/Private/AssetTypeActions/AssetTypeActions_World.cpp
--- a/Engine/Source/Developer/AssetTools/Private/AssetTypeActions/AssetTypeActions_World.cpp
+++ b/Engine/Source/Developer/AssetTools/Private/AssetTypeActions/AssetTypeActions_World.cpp
@@ -12,7 +12,7 @@ void FAssetTypeActions_World::OpenAssetEditor( const TArray<UObject*>& InObjects
 	for (auto ObjIt = InObjects.CreateConstIterator(); ObjIt; ++ObjIt)
 	{
 		UWorld* World = Cast<UWorld>(*ObjIt);
-		if (World != nullptr && ensure(World->IsValidLowLevel()))
+		if (World != nullptr && ensureMsgf(World->GetTypedOuter<UPackage>(), TEXT("World(%s) is not in a package and cannot be opened"), *World->GetFullName()))
 		{
 			// If there are any unsaved changes to the current level, see if the user wants to save those first.
 			bool bPromptUserToSave = true;
diff --git a/Engine/Source/Developer/AutomationController/Private/AutomationControllerManger.cpp b/Engine/Source/Developer/AutomationController/Private/AutomationControllerManger.cpp
--- a/Engine/Source/Developer/AutomationController/Private/AutomationControllerManger.cpp
+++ b/Engine/Source/Developer/AutomationController/Private/AutomationControllerManger.cpp
@@ -751,7 +751,7 @@ void FAutomationControllerManager::HandleReceivedScreenShot( const FAutomationWo
 	if ( FJsonObjectConverter::UStructToJsonObjectString(Message.Metadata, Json) )
 	{
 		FString MetadataPath = FPaths::ChangeExtension(FileName, TEXT("json"));
-		FFileHelper::SaveStringToFile(Json, *MetadataPath, FFileHelper::EEncodingOptions::ForceUTF8);
+		FFileHelper::SaveStringToFile(Json, *MetadataPath, FFileHelper::EEncodingOptions::ForceUTF8WithoutBOM);
 	}
 
 	TSharedRef<FComparisonEntry> Comparison = MakeShareable(new FComparisonEntry());
diff --git a/Engine/Source/Developer/FunctionalTesting/Private/AutomationBlueprintFunctionLibrary.cpp b/Engine/Source/Developer/FunctionalTesting/Private/AutomationBlueprintFunctionLibrary.cpp
--- a/Engine/Source/Developer/FunctionalTesting/Private/AutomationBlueprintFunctionLibrary.cpp
+++ b/Engine/Source/Developer/FunctionalTesting/Private/AutomationBlueprintFunctionLibrary.cpp
@@ -21,6 +21,8 @@
 #include "DelayAction.h"
 #include "Widgets/SViewport.h"
 #include "Framework/Application/SlateApplication.h"
+#include "ShaderCompiler.h"
+#include "AutomationBlueprintFunctionLibrary.h"
 
 #define LOCTEXT_NAMESPACE "Automation"
 
@@ -251,6 +253,9 @@ bool UAutomationBlueprintFunctionLibrary::TakeAutomationScreenshotInternal(const
 	// Force all mip maps to load before taking the screenshot.
 	UTexture::ForceUpdateTextureStreaming();
 
+	// Force all shader compiling to finish.
+	GShaderCompilingManager->FinishAllCompilation();
+
 #if (WITH_DEV_AUTOMATION_TESTS || WITH_PERF_AUTOMATION_TESTS)
 	FAutomationScreenshotTaker* TempObject = new FAutomationScreenshotTaker(Name, Options);
 #endif
diff --git a/Engine/Source/Developer/ScreenShotComparison/Private/Models/ScreenComparisonModel.h b/Engine/Source/Developer/ScreenShotComparison/Private/Models/ScreenComparisonModel.h
--- a/Engine/Source/Developer/ScreenShotComparison/Private/Models/ScreenComparisonModel.h
+++ b/Engine/Source/Developer/ScreenShotComparison/Private/Models/ScreenComparisonModel.h
@@ -3,6 +3,7 @@
 
 #include "CoreMinimal.h"
 #include "ImageComparer.h"
+#include "AutomationWorkerMessages.h"
 
 enum class EComparisonResultType
 {
@@ -25,6 +26,8 @@ public:
 
 	TOptional<FImageComparisonResult> ComparisonResult;
 
+	TOptional<FAutomationScreenshotMetadata> Metadata;
+
 private:
 
 	EComparisonResultType Type;
diff --git a/Engine/Source/Developer/ScreenShotComparison/Private/Widgets/SScreenComparisonRow.cpp b/Engine/Source/Developer/ScreenShotComparison/Private/Widgets/SScreenComparisonRow.cpp
--- a/Engine/Source/Developer/ScreenShotComparison/Private/Widgets/SScreenComparisonRow.cpp
+++ b/Engine/Source/Developer/ScreenShotComparison/Private/Widgets/SScreenComparisonRow.cpp
@@ -17,6 +17,8 @@
 #include "Misc/FileHelper.h"
 #include "Interfaces/IImageWrapperModule.h"
 #include "Framework/Application/SlateApplication.h"
+#include "JsonObjectConverter.h"
+#include "SScreenShotImagePopup.h"
 
 #define LOCTEXT_NAMESPACE "SScreenShotBrowser"
 
@@ -35,12 +37,11 @@ void SScreenComparisonRow::Construct( const FArguments& InArgs, const TSharedRef
 		{
 			FString AddedFolder = ComparisonDirectory / Comparisons->IncomingPath / Model->Folder;
 			IFileManager::Get().FindFilesRecursive(ExternalFiles, *AddedFolder, TEXT("*.*"), true, false);
+			LoadMetadata();
 			break;
 		}
 		case EComparisonResultType::Missing:
 		{
-			FString AddedFolder = ComparisonDirectory / Comparisons->ApprovedPath / Model->Folder;
-			IFileManager::Get().FindFilesRecursive(ExternalFiles, *AddedFolder, TEXT("*.*"), true, false);
 			break;
 		}
 		case EComparisonResultType::Compared:
@@ -48,6 +49,7 @@ void SScreenComparisonRow::Construct( const FArguments& InArgs, const TSharedRef
 			const FImageComparisonResult& ComparisonResult = Model->ComparisonResult.GetValue();
 			FString IncomingImage = ComparisonDirectory / Comparisons->IncomingPath / ComparisonResult.IncomingFile;
 			FString IncomingMetadata = FPaths::ChangeExtension(IncomingImage, TEXT("json"));
+			LoadMetadata();
 
 			ExternalFiles.Add(IncomingImage);
 			ExternalFiles.Add(IncomingMetadata);
@@ -58,11 +60,56 @@ void SScreenComparisonRow::Construct( const FArguments& InArgs, const TSharedRef
 	SMultiColumnTableRow<TSharedPtr<FScreenComparisonModel>>::Construct(FSuperRowType::FArguments(), InOwnerTableView);
 }
 
+void SScreenComparisonRow::LoadMetadata()
+{
+	if ( !Model->Metadata.IsSet() )
+	{
+		FString IncomingMetadataFile;
+
+		if ( Model->ComparisonResult.IsSet() )
+		{
+			const FImageComparisonResult& ComparisonResult = Model->ComparisonResult.GetValue();
+			FString IncomingImage = ComparisonDirectory / Comparisons->IncomingPath / ComparisonResult.IncomingFile;
+			IncomingMetadataFile = FPaths::ChangeExtension(IncomingImage, TEXT("json"));
+		}
+		else
+		{
+			TArray<FString> MetadataFiles;
+			FString AddedFolder = ComparisonDirectory / Comparisons->IncomingPath / Model->Folder;
+			IFileManager::Get().FindFilesRecursive(MetadataFiles, *AddedFolder, TEXT("*.json"), true, false);
+			if ( MetadataFiles.Num() > 0 )
+			{
+				IncomingMetadataFile = MetadataFiles[0];
+			}
+		}
+
+		if ( !IncomingMetadataFile.IsEmpty() )
+		{
+			FString Json;
+			if ( FFileHelper::LoadFileToString(Json, *IncomingMetadataFile) )
+			{
+				FAutomationScreenshotMetadata Metadata;
+				if ( FJsonObjectConverter::JsonObjectStringToUStruct(Json, &Metadata, 0, 0) )
+				{
+					Model->Metadata = Metadata;
+				}
+			}
+		}
+	}
+}
+
 TSharedRef<SWidget> SScreenComparisonRow::GenerateWidgetForColumn(const FName& ColumnName)
 {
 	if ( ColumnName == "Name" )
 	{
-		return SNew(STextBlock).Text(LOCTEXT("Unknown", "Test Name?  Driver?"));
+		if ( Model->Metadata.IsSet() )
+		{
+			return SNew(STextBlock).Text(FText::FromString(Model->Metadata->Name));
+		}
+		else
+		{
+			return SNew(STextBlock).Text(LOCTEXT("Unknown", "Unknown Test, no metadata discovered."));
+		}
 	}
 	else if ( ColumnName == "Delta" )
 	{
@@ -107,10 +154,26 @@ TSharedRef<SWidget> SScreenComparisonRow::GenerateWidgetForColumn(const FName& C
 					.AutoHeight()
 					.HAlign(HAlign_Center)
 					[
-						SNew(SButton)
-						.IsEnabled(this, &SScreenComparisonRow::CanUseSourceControl)
-						.Text(LOCTEXT("ApproveNew", "Approve New!"))
-						.OnClicked(this, &SScreenComparisonRow::ReplaceOld)
+						SNew(SHorizontalBox)
+
+						+ SHorizontalBox::Slot()
+						.AutoWidth()
+						[
+							SNew(SButton)
+							.IsEnabled(this, &SScreenComparisonRow::CanReplace)
+							.Text(LOCTEXT("Replace", "Replace"))
+							.OnClicked(this, &SScreenComparisonRow::Replace)
+						]
+
+						+ SHorizontalBox::Slot()
+						.AutoWidth()
+						.Padding(10, 0, 0, 0)
+						[
+							SNew(SButton)
+							.IsEnabled(this, &SScreenComparisonRow::CanAddAsAlternative)
+							.Text(LOCTEXT("AddAlternative", "Add As Alternative"))
+							.OnClicked(this, &SScreenComparisonRow::AddAlternative)
+						]
 					];
 			}
 		}
@@ -129,25 +192,59 @@ TSharedRef<SWidget> SScreenComparisonRow::BuildMissingView()
 	return
 		SNew(SHorizontalBox)
 		
-		+SHorizontalBox::Slot()
-		[
-			SNew(SButton)
-			.IsEnabled(this, &SScreenComparisonRow::CanUseSourceControl)
-			.Text(LOCTEXT("RemoveOld", "Remove Old"))
-			.OnClicked(this, &SScreenComparisonRow::RemoveOld)
-		];
+		/*	+SHorizontalBox::Slot()
+			[
+				SNew(SButton)
+				.IsEnabled(this, &SScreenComparisonRow::CanUseSourceControl)
+				.Text(LOCTEXT("RemoveOld", "Remove Old"))
+				.OnClicked(this, &SScreenComparisonRow::RemoveOld)
+			]*/;
 }
 
 TSharedRef<SWidget> SScreenComparisonRow::BuildAddedView()
 {
+	TArray<FString> Images;
+
+	FString AddedFolder = ComparisonDirectory / Comparisons->IncomingPath / Model->Folder;
+	IFileManager::Get().FindFilesRecursive(ExternalFiles, *AddedFolder, TEXT("*.png"), true, false);
+
+	//TODO Automation this is no good, we need to clear out the incoming directory beforehand.
+	if ( Images.Num() > 0 )
+	{
+		UnapprovedBrush = LoadScreenshot(Images[0]);
+	}
+
 	return
-		SNew(SHorizontalBox)
+		SNew(SVerticalBox)
+		
+		+ SVerticalBox::Slot()
+		[
+			SNew(SBox)
+			.MaxDesiredHeight(100)
+			.HAlign(HAlign_Left)
+			[
+				SNew(SScaleBox)
+				.Stretch(EStretch::ScaleToFit)
+				[
+					SNew(SHorizontalBox)
+		
+					+ SHorizontalBox::Slot()
+					.AutoWidth()
+					.Padding(4.0f, 4.0f)
+					[
+						SNew(SImage)
+						.Image(UnapprovedBrush.Get())
+						.OnMouseButtonDown(this, &SScreenComparisonRow::OnImageClicked, UnapprovedBrush)
+					]
+				]
+			]
+		]
 		
-		+SHorizontalBox::Slot()
+		+ SVerticalBox::Slot()
 		[
 			SNew(SButton)
 			.IsEnabled(this, &SScreenComparisonRow::CanUseSourceControl)
-			.Text(LOCTEXT("AddComparison", "Add New!"))
+			.Text(LOCTEXT("AddNew", "Add New!"))
 			.OnClicked(this, &SScreenComparisonRow::AddNew)
 		];
 }
@@ -181,6 +278,7 @@ TSharedRef<SWidget> SScreenComparisonRow::BuildComparisonPreview()
 				[
 					SNew(SImage)
 					.Image(ApprovedBrush.Get())
+					.OnMouseButtonDown(this, &SScreenComparisonRow::OnImageClicked, ApprovedBrush)
 				]
 
 				+ SHorizontalBox::Slot()
@@ -189,6 +287,7 @@ TSharedRef<SWidget> SScreenComparisonRow::BuildComparisonPreview()
 				[
 					SNew(SImage)
 					.Image(ComparisonBrush.Get())
+					.OnMouseButtonDown(this, &SScreenComparisonRow::OnImageClicked, ComparisonBrush)
 				]
 
 				+ SHorizontalBox::Slot()
@@ -197,6 +296,7 @@ TSharedRef<SWidget> SScreenComparisonRow::BuildComparisonPreview()
 				[
 					SNew(SImage)
 					.Image(UnapprovedBrush.Get())
+					.OnMouseButtonDown(this, &SScreenComparisonRow::OnImageClicked, UnapprovedBrush)
 				]
 			]
 		];
@@ -211,6 +311,11 @@ void SScreenComparisonRow::GetStatus()
 	//}
 }
 
+bool SScreenComparisonRow::CanAddNew() const
+{
+	return CanUseSourceControl();
+}
+
 FReply SScreenComparisonRow::AddNew()
 {
 	// Copy files to the approved
@@ -241,50 +346,100 @@ FReply SScreenComparisonRow::AddNew()
 	return FReply::Handled();
 }
 
-FReply SScreenComparisonRow::RemoveOld()
+bool SScreenComparisonRow::CanReplace() const
+{
+	return CanUseSourceControl();
+}
+
+FReply SScreenComparisonRow::RemoveExistingApproved()
 {
+	TArray<FString> FilesToRemove;
+
+	FString PlatformFolder = ComparisonDirectory / Comparisons->ApprovedPath / FPaths::GetPath(Model->ComparisonResult->ApprovedFile);
+	IFileManager::Get().FindFilesRecursive(FilesToRemove, *PlatformFolder, TEXT("*.*"), true, false);
+
 	// Copy files to the approved
 	const FString& LocalApprovedFolder = ScreenshotManager->GetLocalApprovedFolder();
 	const FString ImportApprovedRoot = ComparisonDirectory / Comparisons->ApprovedPath / TEXT("");
 
 	TArray<FString> SourceControlFiles;
 
-	for ( const FString& File : ExternalFiles )
+	for ( const FString& File : FilesToRemove )
 	{
-		FString RelativeFile = ExternalFiles[0];
+		FString RelativeFile = File;
 		FPaths::MakePathRelativeTo(RelativeFile, *ImportApprovedRoot);
 
 		FString DestFilePath = LocalApprovedFolder / RelativeFile;
 		SourceControlFiles.Add(DestFilePath);
+
+		IFileManager::Get().Delete(*DestFilePath, false, true, false);
 	}
 
-	if ( ExternalFiles.Num() > 0 )
+	ISourceControlProvider& SourceControlProvider = ISourceControlModule::Get().GetProvider();
+	if ( SourceControlProvider.Execute(ISourceControlOperation::Create<FDelete>(), SourceControlFiles) == ECommandResult::Failed )
 	{
-		FString RelativeFile = ExternalFiles[0];
-		FPaths::MakePathRelativeTo(RelativeFile, *ImportApprovedRoot);
+		//TODO Error
+	}
+
+	// Invalidate Status
 
-		FString TestFolder = FPaths::GetPath(FPaths::GetPath(RelativeFile));
+	return FReply::Handled();
+}
 
-		//FString ParentDir = SelectedPaths[0] / TEXT("..");
-		//FPaths::CollapseRelativeDirectories(ParentDir);
+FReply SScreenComparisonRow::Replace()
+{
+	// Delete all the existing files in this area
+	RemoveExistingApproved();
 
-		FString MissingLocalTest = LocalApprovedFolder / TestFolder;
+	// Copy files to the approved
+	const FString& LocalApprovedFolder = ScreenshotManager->GetLocalApprovedFolder();
+	const FString ImportIncomingRoot = ComparisonDirectory / Comparisons->IncomingPath / TEXT("");
+
+	TArray<FString> SourceControlFiles;
 
-		IFileManager::Get().DeleteDirectory(*MissingLocalTest, false, true);
+	for ( const FString& File : ExternalFiles )
+	{
+		FString RelativeFile = File;
+		FPaths::MakePathRelativeTo(RelativeFile, *ImportIncomingRoot);
+
+		FString DestFilePath = LocalApprovedFolder / RelativeFile;
+		
+		SourceControlFiles.Add(DestFilePath);
 	}
 
 	ISourceControlProvider& SourceControlProvider = ISourceControlModule::Get().GetProvider();
-	if ( SourceControlProvider.Execute(ISourceControlOperation::Create<FDelete>(), SourceControlFiles) == ECommandResult::Failed )
+	if ( SourceControlProvider.Execute(ISourceControlOperation::Create<FRevert>(), SourceControlFiles) == ECommandResult::Failed )
 	{
 		//TODO Error
 	}
 
-	// Invalidate Status
+	for ( const FString& File : ExternalFiles )
+	{
+		FString RelativeFile = File;
+		FPaths::MakePathRelativeTo(RelativeFile, *ImportIncomingRoot);
+
+		FString DestFilePath = LocalApprovedFolder / RelativeFile;
+		IFileManager::Get().Copy(*DestFilePath, *File, true, true);
+	}
+
+	if ( SourceControlProvider.Execute(ISourceControlOperation::Create<FMarkForAdd>(), SourceControlFiles) == ECommandResult::Failed )
+	{
+		//TODO Error
+	}
+	if ( SourceControlProvider.Execute(ISourceControlOperation::Create<FCheckOut>(), SourceControlFiles) == ECommandResult::Failed )
+	{
+		//TODO Error
+	}
 
 	return FReply::Handled();
 }
 
-FReply SScreenComparisonRow::ReplaceOld()
+bool SScreenComparisonRow::CanAddAsAlternative() const
+{
+	return CanUseSourceControl() && (Model->ComparisonResult->IncomingFile != Model->ComparisonResult->ApprovedFile);
+}
+
+FReply SScreenComparisonRow::AddAlternative()
 {
 	// Copy files to the approved
 	const FString& LocalApprovedFolder = ScreenshotManager->GetLocalApprovedFolder();
@@ -298,12 +453,29 @@ FReply SScreenComparisonRow::ReplaceOld()
 		FPaths::MakePathRelativeTo(RelativeFile, *ImportIncomingRoot);
 
 		FString DestFilePath = LocalApprovedFolder / RelativeFile;
-		IFileManager::Get().Copy(*DestFilePath, *File, true, true);
 
 		SourceControlFiles.Add(DestFilePath);
 	}
 
 	ISourceControlProvider& SourceControlProvider = ISourceControlModule::Get().GetProvider();
+	if ( SourceControlProvider.Execute(ISourceControlOperation::Create<FRevert>(), SourceControlFiles) == ECommandResult::Failed )
+	{
+		//TODO Error
+	}
+
+	for ( const FString& File : ExternalFiles )
+	{
+		FString RelativeFile = File;
+		FPaths::MakePathRelativeTo(RelativeFile, *ImportIncomingRoot);
+
+		FString DestFilePath = LocalApprovedFolder / RelativeFile;
+		IFileManager::Get().Copy(*DestFilePath, *File, true, true);
+	}
+
+	if ( SourceControlProvider.Execute(ISourceControlOperation::Create<FMarkForAdd>(), SourceControlFiles) == ECommandResult::Failed )
+	{
+		//TODO Error
+	}
 	if ( SourceControlProvider.Execute(ISourceControlOperation::Create<FCheckOut>(), SourceControlFiles) == ECommandResult::Failed )
 	{
 		//TODO Error
@@ -347,4 +519,33 @@ TSharedPtr<FSlateDynamicImageBrush> SScreenComparisonRow::LoadScreenshot(FString
 	return nullptr;
 }
 
-#undef LOCTEXT_NAMESPACE
+FReply SScreenComparisonRow::OnImageClicked(const FGeometry& InGeometry, const FPointerEvent& InEvent, TSharedPtr<FSlateDynamicImageBrush> Image)
+{
+	TSharedRef<SScreenShotImagePopup> PopupImage =
+		SNew(SScreenShotImagePopup)
+		.ImageBrush(Image)
+		.ImageSize(Image->ImageSize.IntPoint());
+
+	auto ParentWindow = FSlateApplication::Get().FindWidgetWindow(AsShared()).ToSharedRef();
+
+	// Center ourselves in the parent window
+	auto PopupWindow = SNew(SWindow)
+		.IsPopupWindow(false)
+		.SizingRule(ESizingRule::UserSized)
+		.ClientSize(Image->ImageSize)
+		.AutoCenter(EAutoCenter::PreferredWorkArea)
+		.SupportsMaximize(true)
+		.SupportsMinimize(true)
+		.FocusWhenFirstShown(true)
+		.ActivateWhenFirstShown(true)
+		.Content()
+		[
+			PopupImage
+		];
+
+	FSlateApplication::Get().AddWindowAsNativeChild(PopupWindow, ParentWindow, true);
+
+	return FReply::Handled();
+}
+
+#undef LOCTEXT_NAMESPACE
\ No newline at end of file
diff --git a/Engine/Source/Developer/ScreenShotComparison/Private/Widgets/SScreenComparisonRow.h b/Engine/Source/Developer/ScreenShotComparison/Private/Widgets/SScreenComparisonRow.h
--- a/Engine/Source/Developer/ScreenShotComparison/Private/Widgets/SScreenComparisonRow.h
+++ b/Engine/Source/Developer/ScreenShotComparison/Private/Widgets/SScreenComparisonRow.h
@@ -48,13 +48,23 @@ private:
 	TSharedRef<SWidget> BuildAddedView();
 	TSharedRef<SWidget> BuildComparisonPreview();
 
+	bool CanAddNew() const;
 	FReply AddNew();
-	FReply RemoveOld();
-	FReply ReplaceOld();
+
+	bool CanReplace() const;
+	FReply Replace();
+
+	bool CanAddAsAlternative() const;
+	FReply AddAlternative();
+
+	FReply RemoveExistingApproved();
 
 	void GetStatus();
 
 	TSharedPtr<FSlateDynamicImageBrush> LoadScreenshot(FString ImagePath);
+	void LoadMetadata();
+
+	FReply OnImageClicked(const FGeometry& InGeometry, const FPointerEvent& InEvent, TSharedPtr<FSlateDynamicImageBrush> Image);
 
 private:
 
diff --git a/Engine/Source/Developer/ScreenShotComparison/Private/Widgets/SScreenShotBrowser.cpp b/Engine/Source/Developer/ScreenShotComparison/Private/Widgets/SScreenShotBrowser.cpp
--- a/Engine/Source/Developer/ScreenShotComparison/Private/Widgets/SScreenShotBrowser.cpp
+++ b/Engine/Source/Developer/ScreenShotComparison/Private/Widgets/SScreenShotBrowser.cpp
@@ -140,12 +140,13 @@ void SScreenShotBrowser::RebuildTree()
 				ComparisonList.Add(Model);
 			}
 
-			for ( FString& Missing : CurrentComparisons->Missing )
-			{
-				TSharedPtr<FScreenComparisonModel> Model = MakeShared<FScreenComparisonModel>(EComparisonResultType::Missing);
-				Model->Folder = Missing;
-				ComparisonList.Add(Model);
-			}
+			// Ignore Missing
+			//for ( FString& Missing : CurrentComparisons->Missing )
+			//{
+			//	TSharedPtr<FScreenComparisonModel> Model = MakeShared<FScreenComparisonModel>(EComparisonResultType::Missing);
+			//	Model->Folder = Missing;
+			//	ComparisonList.Add(Model);
+			//}
 
 			// Copy the comparisons to an array as shared pointers the list view can use.
 			for ( FImageComparisonResult& Result : CurrentComparisons->Comparisons )
diff --git a/Engine/Source/Developer/ScreenShotComparison/Private/Widgets/SScreenShotImagePopup.cpp b/Engine/Source/Developer/ScreenShotComparison/Private/Widgets/SScreenShotImagePopup.cpp
--- a/Engine/Source/Developer/ScreenShotComparison/Private/Widgets/SScreenShotImagePopup.cpp
+++ b/Engine/Source/Developer/ScreenShotComparison/Private/Widgets/SScreenShotImagePopup.cpp
@@ -13,7 +13,7 @@
 void SScreenShotImagePopup::Construct( const FArguments& InArgs )
 {
 	// Load the image
-	DynamicImageBrush = MakeShareable(new FSlateDynamicImageBrush( InArgs._ImageAssetName, InArgs._ImageSize ) );
+	DynamicImageBrush = InArgs._ImageBrush;
 
 	// Create the screen shot popup widget.
 	ChildSlot
diff --git a/Engine/Source/Developer/ScreenShotComparison/Private/Widgets/SScreenShotImagePopup.h b/Engine/Source/Developer/ScreenShotComparison/Private/Widgets/SScreenShotImagePopup.h
--- a/Engine/Source/Developer/ScreenShotComparison/Private/Widgets/SScreenShotImagePopup.h
+++ b/Engine/Source/Developer/ScreenShotComparison/Private/Widgets/SScreenShotImagePopup.h
@@ -21,7 +21,7 @@ public:
 	SLATE_BEGIN_ARGS( SScreenShotImagePopup ) {}
 	
 	/** The assest name for the screen shot */
-	SLATE_ARGUMENT( FName, ImageAssetName )
+	SLATE_ARGUMENT( TSharedPtr<FSlateDynamicImageBrush>, ImageBrush )
 
 	/** The size of the image */
 	SLATE_ARGUMENT( FIntPoint, ImageSize )
diff --git a/Engine/Source/Developer/ScreenShotComparison/ScreenShotComparison.Build.cs b/Engine/Source/Developer/ScreenShotComparison/ScreenShotComparison.Build.cs
--- a/Engine/Source/Developer/ScreenShotComparison/ScreenShotComparison.Build.cs
+++ b/Engine/Source/Developer/ScreenShotComparison/ScreenShotComparison.Build.cs
@@ -23,6 +23,9 @@ public class ScreenShotComparison : ModuleRules
 				"CoreUObject",
 				"DesktopWidgets",
 				"SourceControl",
+				"AutomationMessages",
+				"Json",
+				"JsonUtilities"
 			}
 		);
 
diff --git a/Engine/Source/Developer/ScreenShotComparisonTools/Private/ImageComparer.cpp b/Engine/Source/Developer/ScreenShotComparisonTools/Private/ImageComparer.cpp
--- a/Engine/Source/Developer/ScreenShotComparisonTools/Private/ImageComparer.cpp
+++ b/Engine/Source/Developer/ScreenShotComparisonTools/Private/ImageComparer.cpp
@@ -47,7 +47,7 @@ public:
 		int32 Offset = ( Y * Width + X ) * 4;
 		check(Offset < ( Width * Height * 4 ));
 
-		const float Brightness = FPixelOperations::GetBrightness(Color);
+		const double Brightness = FPixelOperations::GetLuminance(Color);
 
 		Image[Offset] = Brightness;
 		Image[Offset + 1] = Brightness;
@@ -161,7 +161,7 @@ bool FPixelOperations::IsAntialiased(const FColor& SourcePixel, FComparableImage
 
 				FColor TargetPixel = Image->GetPixel(X + j, Y + i);
 
-				float TargetPixelBrightness = GetBrightness(TargetPixel);
+				double TargetPixelBrightness = GetLuminance(TargetPixel);
 				float TargetPixelHue = GetHue(TargetPixel);
 
 				if ( FPixelOperations::IsContrasting(SourcePixel, TargetPixel, Tolerance) )
@@ -203,23 +203,23 @@ void FComparableImage::Process()
 		for ( int Y = 0; Y < Height; Y++ )
 		{
 			FColor Pixel = GetPixel(ColumnIndex, Y);
-			float Brightness = FPixelOperations::GetBrightness(Pixel);
+			double Luminance = FPixelOperations::GetLuminance(Pixel);
 
 			RedTotal += ( Pixel.R / 255.0 );
 			GreenTotal += ( Pixel.G / 255.0 );
 			BlueTotal += ( Pixel.B / 255.0 );
 			AlphaTotal += ( Pixel.A / 255.0 );
-			BrightnessTotal += ( Brightness / 255.0 );
+			LuminanceTotal += ( Luminance / 255.0 );
 		}
 	});
 
-	double PixelCount = Width * Height;
+	const double PixelCount = Width * Height;
 
 	RedAverage = RedTotal / PixelCount;
 	GreenAverage = GreenTotal / PixelCount;
 	BlueAverage = BlueTotal / PixelCount;
 	AlphaAverage = AlphaTotal / PixelCount;
-	BbrightnessAverage = BrightnessTotal / PixelCount;
+	LuminanceAverage = LuminanceTotal / PixelCount;
 }
 
 TSharedPtr<FComparableImage> FImageComparer::Open(const FString& ImagePath, FText& OutError)
@@ -406,4 +406,163 @@ FImageComparisonResult FImageComparer::Compare(const FString& ImagePathA, const
 	return Results;
 }
 
+double FImageComparer::CompareStructuralSimilarity(const FString& ImagePathA, const FString& ImagePathB, EStructuralSimilarityComponent InCompareComponent)
+{
+	FImageComparisonResult Results;
+	Results.ApprovedFile = ImagePathA;
+	FPaths::MakePathRelativeTo(Results.ApprovedFile, *ImageRootA);
+	Results.IncomingFile = ImagePathB;
+	FPaths::MakePathRelativeTo(Results.IncomingFile, *ImageRootB);
+
+	TSharedPtr<FComparableImage> ImageA;
+	TSharedPtr<FComparableImage> ImageB;
+
+	FText ErrorA;
+	FText ErrorB;
+
+	ParallelFor(2,
+		[&] (int32 Index)
+	{
+		if ( Index == 0 )
+		{
+			ImageA = Open(ImagePathA, ErrorA);
+		}
+		else
+		{
+			ImageB = Open(ImagePathB, ErrorB);
+		}
+	});
+
+	if ( !ImageA.IsValid() )
+	{
+		Results.ErrorMessage = ErrorA;
+		return 0.0f;
+	}
+
+	if ( !ImageB.IsValid() )
+	{
+		Results.ErrorMessage = ErrorB;
+		return 0.0f;
+	}
+
+	if ( ImageA->Width != ImageB->Width || ImageA->Height != ImageB->Height )
+	{
+		Results.ErrorMessage = LOCTEXT("DifferentSizesUnsupported", "We can not compare images of different sizes at this time.");
+		return 0.0f;
+	}
+
+	//ImageA->Process();
+	//ImageB->Process();
+
+	// Implementation of https://en.wikipedia.org/wiki/Structural_similarity
+
+	const double K1 = 0.01;
+	const double K2 = 0.03;
+
+	const int32 BitsPerComponent = 8;
+
+	const int32 MaxWindowSize = 8;
+
+	const int32 ImageWidth = ImageA->Width;
+	const int32 ImageHeight = ImageA->Height;
+
+	int32 TotalWindows = 0;
+	double TotalSSIM = 0;
+
+	for ( int32 X = 0; X < ImageWidth; X += MaxWindowSize )
+	{
+		for ( int32 Y = 0; Y < ImageHeight; Y += MaxWindowSize )
+		{
+			const int32 WindowWidth = FMath::Min(MaxWindowSize, ImageWidth - X);
+			const int32 WindowHeight = FMath::Min(MaxWindowSize, ImageHeight - Y);
+			const int32 WindowEndX = X + WindowWidth;
+			const int32 WindowEndY = Y + WindowHeight;
+
+			double AverageA = 0;
+			double AverageB = 0;
+			double VarianceA = 0;
+			double VarianceB = 0;
+			double CovarianceAB = 0;
+
+			TArray<double> ComponentA;
+			TArray<double> ComponentB;
+
+			// Run through the window, accumulate an average and the components for the window.
+			for ( int32 WindowX = X; WindowX < WindowEndX; WindowX++ )
+			{
+				for ( int32 WindowY = Y; WindowY < WindowEndY; WindowY++ )
+				{
+					if ( InCompareComponent == EStructuralSimilarityComponent::Luminance )
+					{
+						const double LuminanceA = FPixelOperations::GetLuminance(ImageA->GetPixel(WindowX, WindowY));
+						const double LuminanceB = FPixelOperations::GetLuminance(ImageB->GetPixel(WindowX, WindowY));
+						AverageA += LuminanceA;
+						AverageB += LuminanceB;
+						ComponentA.Add(LuminanceA);
+						ComponentB.Add(LuminanceB);
+					}
+					else // EStructuralSimilarityComponent::Color
+					{
+						const FColor ColorA = ImageA->GetPixel(WindowX, WindowY);
+						const FColor ColorB = ImageA->GetPixel(WindowX, WindowY);
+						const double ColorLumpA = ( ColorA.R + ColorA.G + ColorA.B ) * ( ColorA.A / 255.0 );
+						const double ColorLumpB = ( ColorB.R + ColorB.G + ColorB.B ) * ( ColorB.A / 255.0 );
+						AverageA += ColorLumpA;
+						AverageB += ColorLumpB;
+						ComponentA.Add(ColorLumpA);
+						ComponentB.Add(ColorLumpB);
+					}
+				}
+			}
+
+			int32 WindowComponentCount = WindowWidth * WindowHeight;
+
+			// Finally calculate the average.
+			AverageA /= (double)WindowComponentCount;
+			AverageB /= (double)WindowComponentCount;
+
+			check(ComponentA.Num() == WindowComponentCount);
+			check(ComponentB.Num() == WindowComponentCount);
+
+			// Compute the Variance of A and B
+			for ( int32 ComponentIndex = 0; ComponentIndex < WindowComponentCount; ComponentIndex++ )
+			{
+				const double DifferenceA = ComponentA[ComponentIndex] - AverageA;
+				const double DifferenceB = ComponentB[ComponentIndex] - AverageB;
+
+				const double SquaredDifferenceA = FMath::Pow(DifferenceA, 2);
+				const double SquaredDifferenceB = FMath::Pow(DifferenceB, 2);
+
+				VarianceA += SquaredDifferenceA;
+				VarianceB += SquaredDifferenceB;
+				CovarianceAB += DifferenceA * DifferenceB;
+			}
+
+			// Finally divide by the number of components to get the average of the mean of the squared differences (aka variance).
+			VarianceA /= (double)WindowComponentCount;
+			VarianceB /= (double)WindowComponentCount;
+			CovarianceAB /= (double)WindowComponentCount;
+
+			// The dynamic range of the pixel values
+			float L = ( 1 << BitsPerComponent ) - 1;
+
+			// Two variables to stabilize the division of a weak denominator.
+			double C1 = FMath::Pow(K1 * L, 2);
+			double C2 = FMath::Pow(K2 * L, 2);
+
+			double Luminance = ( 2 * AverageA * AverageB + C1 ) / ( FMath::Pow(AverageA, 2) + FMath::Pow(AverageB, 2) + C1 );
+			double Contrast = ( 2 * CovarianceAB + C2) / ( VarianceA + VarianceB + C2 );
+
+			double WindowSSIM = Luminance * Contrast;
+
+			TotalSSIM += WindowSSIM;
+			TotalWindows++;
+		}
+	}
+
+	double SSIM = TotalSSIM / TotalWindows;
+	
+	return FMath::Clamp(SSIM, 0.0, 1.0);
+}
+
 #undef LOCTEXT_NAMESPACE
diff --git a/Engine/Source/Developer/ScreenShotComparisonTools/Private/ScreenShotManager.cpp b/Engine/Source/Developer/ScreenShotComparisonTools/Private/ScreenShotManager.cpp
--- a/Engine/Source/Developer/ScreenShotComparisonTools/Private/ScreenShotManager.cpp
+++ b/Engine/Source/Developer/ScreenShotComparisonTools/Private/ScreenShotManager.cpp
@@ -32,6 +32,9 @@ FScreenShotManager::FScreenShotManager()
 	ScreenshotDeltaFolder = FPaths::GameSavedDir() / TEXT("Automation/Delta/");
 	ScreenshotResultsFolder = FPaths::GameSavedDir() / TEXT("Automation/");
 	ScreenshotApprovedFolder = FPaths::GameDir() / TEXT("Test/Screenshots/");
+
+	// Clear the incoming directory when we initialize, we don't care about previous runs.
+	//IFileManager::Get().DeleteDirectory(*ScreenshotUnapprovedFolder, false, true);
 }
 
 FString FScreenShotManager::GetLocalUnapprovedFolder() const
@@ -235,57 +238,85 @@ FImageComparisonResult FScreenShotManager::CompareScreensot(FString ExistingImag
 		return FImageComparisonResult();
 	}
 
-	// Look for an exact adapter match so we're comparing apples to apples.
-	
+	// Load the metadata for the incoming unapproved image.
 	FString UnapprovedFileName = FPaths::GetCleanFilename(ExistingImage);
+	FString UnapprovedFullPath = FPaths::Combine(TestUnapprovedFolder, UnapprovedFileName);
+
+	TOptional<FAutomationScreenshotMetadata> ExistingMetadata;
 
-	FString* ExistingApprovedAdapter = ApprovedDeviceShots.FindByPredicate([&UnapprovedFileName] (const FString& ApprovedPath) {
-		FString ApprovedFileName = FPaths::GetCleanFilename(ApprovedPath);
-		return ApprovedFileName == UnapprovedFileName;
-	});
+	{
+		// Always read the metadata file from the unapproved location, as we may have introduced new comparison rules.
+		FString MetadataFile = FPaths::ChangeExtension(UnapprovedFullPath, ".json");
+
+		FString Json;
+		if ( FFileHelper::LoadFileToString(Json, *MetadataFile) )
+		{
+			FAutomationScreenshotMetadata Metadata;
+			if ( FJsonObjectConverter::JsonObjectStringToUStruct(Json, &Metadata, 0, 0) )
+			{
+				ExistingMetadata = Metadata;
+			}
+		}
+	}
 
-	FString ExistingApproved;
+	FString NearestExistingApprovedImage;
+	TOptional<FAutomationScreenshotMetadata> NearestExistingApprovedImageMetadata;
+	int32 MatchScore = 0;
 
-	if ( ExistingApprovedAdapter )
+	if ( ExistingMetadata.IsSet() )
 	{
-		ExistingApproved = FPaths::GetCleanFilename(*ExistingApprovedAdapter);
+		for ( FString ApprovedShot : ApprovedDeviceShots )
+		{
+			FString ApprovedShotFile = FPaths::GetCleanFilename(ApprovedDeviceShots[0]);
+			FString ApprovedShotFileFull = FPaths::Combine(TestApprovedFolder, ApprovedShotFile);
+
+			FString ApprovedMetadataFile = FPaths::ChangeExtension(ApprovedShotFileFull, ".json");
+
+			FString Json;
+			if ( FFileHelper::LoadFileToString(Json, *ApprovedMetadataFile) )
+			{
+				FAutomationScreenshotMetadata Metadata;
+				if ( FJsonObjectConverter::JsonObjectStringToUStruct(Json, &Metadata, 0, 0) )
+				{
+					int32 Comparison = Metadata.Compare(ExistingMetadata.GetValue());
+					if ( Comparison > MatchScore )
+					{
+						MatchScore = Comparison;
+						NearestExistingApprovedImage = ApprovedShotFile;
+						NearestExistingApprovedImageMetadata = Metadata;
+					}
+				}
+			}
+		}
 	}
 	else
 	{
-		ExistingApproved = FPaths::GetCleanFilename(ApprovedDeviceShots[0]);
-		// TODO No good, this should find closest resolution, maybe some other stuff as well?
+		// TODO no metadata how do I pick a good shot?
+		NearestExistingApprovedImage = FPaths::GetCleanFilename(ApprovedDeviceShots[0]);
 	}
 
-	FString ApprovedFullPath = FPaths::Combine(TestApprovedFolder, ExistingApproved);
-	FString UnapprovedFullPath = FPaths::Combine(TestUnapprovedFolder, UnapprovedFileName);
-
-	// Always read the metadata file from the unapproved location, as we may have introduced new comparison rules.
-	FString MetadataFile = FPaths::ChangeExtension(UnapprovedFullPath, ".json");
+	FString ApprovedFullPath = FPaths::Combine(TestApprovedFolder, NearestExistingApprovedImage);
 
 	FImageTolerance Tolerance = DefaultTolerance;
 
-	FString Json;
-	if ( FFileHelper::LoadFileToString(Json, *MetadataFile) )
+	if ( ExistingMetadata.IsSet() && ExistingMetadata->bHasComparisonRules )
 	{
-		FAutomationScreenshotMetadata Metadata;
-		if ( FJsonObjectConverter::JsonObjectStringToUStruct(Json, &Metadata, 0, 0) )
-		{
-			if ( Metadata.bHasComparisonRules )
-			{
-				Tolerance.Red = Metadata.ToleranceRed;
-				Tolerance.Green = Metadata.ToleranceGreen;
-				Tolerance.Blue = Metadata.ToleranceBlue;
-				Tolerance.Alpha = Metadata.ToleranceAlpha;
-				Tolerance.MinBrightness = Metadata.ToleranceMinBrightness;
-				Tolerance.MaxBrightness = Metadata.ToleranceMaxBrightness;
-				Tolerance.IgnoreAntiAliasing = Metadata.bIgnoreAntiAliasing;
-				Tolerance.IgnoreColors = Metadata.bIgnoreColors;
-				Tolerance.MaximumLocalError = Metadata.MaximumLocalError;
-				Tolerance.MaximumGlobalError = Metadata.MaximumGlobalError;
-			}
-		}
+		Tolerance.Red = ExistingMetadata->ToleranceRed;
+		Tolerance.Green = ExistingMetadata->ToleranceGreen;
+		Tolerance.Blue = ExistingMetadata->ToleranceBlue;
+		Tolerance.Alpha = ExistingMetadata->ToleranceAlpha;
+		Tolerance.MinBrightness = ExistingMetadata->ToleranceMinBrightness;
+		Tolerance.MaxBrightness = ExistingMetadata->ToleranceMaxBrightness;
+		Tolerance.IgnoreAntiAliasing = ExistingMetadata->bIgnoreAntiAliasing;
+		Tolerance.IgnoreColors = ExistingMetadata->bIgnoreColors;
+		Tolerance.MaximumLocalError = ExistingMetadata->MaximumLocalError;
+		Tolerance.MaximumGlobalError = ExistingMetadata->MaximumGlobalError;
 	}
 
+	// TODO Think about using SSIM, but needs local SSIM as well as Global SSIM, same as the basic comparison.
+	//double SSIM = Comparer.CompareStructuralSimilarity(ApprovedFullPath, UnapprovedFullPath, FImageComparer::EStructuralSimilarityComponent::Luminance);
+	//printf("%f\n", SSIM);
+
 	return Comparer.Compare(ApprovedFullPath, UnapprovedFullPath, Tolerance);
 }
 
@@ -379,7 +410,7 @@ void FScreenShotManager::SaveComparisonResults(const FComparisonResults& Results
 	FString Json;
 	if ( FJsonObjectConverter::UStructToJsonObjectString(Results, Json) )
 	{
-		FFileHelper::SaveStringToFile(Json, *GetComparisonResultsJsonFilename(), FFileHelper::EEncodingOptions::ForceUTF8);
+		FFileHelper::SaveStringToFile(Json, *GetComparisonResultsJsonFilename(), FFileHelper::EEncodingOptions::ForceUTF8WithoutBOM);
 	}
 }
 
diff --git a/Engine/Source/Developer/ScreenShotComparisonTools/Public/ImageComparer.h b/Engine/Source/Developer/ScreenShotComparisonTools/Public/ImageComparer.h
--- a/Engine/Source/Developer/ScreenShotComparisonTools/Public/ImageComparer.h
+++ b/Engine/Source/Developer/ScreenShotComparisonTools/Public/ImageComparer.h
@@ -89,18 +89,18 @@ class FComparableImage;
 class FPixelOperations
 {
 public:
-	static FORCEINLINE float GetBrightness(const FColor& Color)
+	static FORCEINLINE double GetLuminance(const FColor& Color)
 	{
 		// https://en.wikipedia.org/wiki/Relative_luminance
-		return 0.2126 * Color.R + 0.7152 * Color.G + 0.0722 * Color.B;
+		return (0.2126 * Color.R + 0.7152 * Color.G + 0.0722 * Color.B) * (Color.A / 255.0);
 	}
 
 	static bool IsBrightnessSimilar(const FColor& ColorA, const FColor& ColorB, const FImageTolerance& Tolerance)
 	{
 		const bool AlphaSimilar = FMath::IsNearlyEqual((float)ColorA.A, ColorB.A, Tolerance.Alpha);
 
-		const float BrightnessA = FPixelOperations::GetBrightness(ColorA);
-		const float BrightnessB = FPixelOperations::GetBrightness(ColorB);
+		const float BrightnessA = FPixelOperations::GetLuminance(ColorA);
+		const float BrightnessB = FPixelOperations::GetLuminance(ColorB);
 		const bool BrightnessSimilar = FMath::IsNearlyEqual(BrightnessA, BrightnessB, Tolerance.MinBrightness);
 
 		return BrightnessSimilar && AlphaSimilar;
@@ -125,8 +125,8 @@ public:
 
 	static FORCEINLINE bool IsContrasting(const FColor& ColorA, const FColor& ColorB, const FImageTolerance& Tolerance)
 	{
-		const float BrightnessA = FPixelOperations::GetBrightness(ColorA);
-		const float BrightnessB = FPixelOperations::GetBrightness(ColorB);
+		const float BrightnessA = FPixelOperations::GetLuminance(ColorA);
+		const float BrightnessB = FPixelOperations::GetLuminance(ColorB);
 
 		return FMath::Abs(BrightnessA - BrightnessB) > Tolerance.MaxBrightness;
 	}
@@ -151,12 +151,12 @@ public:
 		, GreenTotal(0)
 		, BlueTotal(0)
 		, AlphaTotal(0)
-		, BrightnessTotal(0)
+		, LuminanceTotal(0)
 		, RedAverage(0)
 		, GreenAverage(0)
 		, BlueAverage(0)
 		, AlphaAverage(0)
-		, BbrightnessAverage(0)
+		, LuminanceAverage(0)
 	{
 	}
 
@@ -185,13 +185,13 @@ public:
 	double GreenTotal;
 	double BlueTotal;
 	double AlphaTotal;
-	double BrightnessTotal;
+	double LuminanceTotal;
 
 	double RedAverage;
 	double GreenAverage;
 	double BlueAverage;
 	double AlphaAverage;
-	double BbrightnessAverage;
+	double LuminanceAverage;
 };
 
 /**
@@ -310,6 +310,17 @@ public:
 
 	FImageComparisonResult Compare(const FString& ImagePathA, const FString& ImagePathB, FImageTolerance Tolerance);
 
+	enum class EStructuralSimilarityComponent : uint8
+	{
+		Luminance,
+		Color
+	};
+
+	/**
+	 * https://en.wikipedia.org/wiki/Structural_similarity
+	 */
+	double CompareStructuralSimilarity(const FString& ImagePathA, const FString& ImagePathB, EStructuralSimilarityComponent InCompareComponent);
+
 private:
 	TSharedPtr<FComparableImage> Open(const FString& ImagePath, FText& OutError);
 };
diff --git a/Engine/Source/Developer/SettingsEditor/Private/SettingsEditorModule.cpp b/Engine/Source/Developer/SettingsEditor/Private/SettingsEditorModule.cpp
--- a/Engine/Source/Developer/SettingsEditor/Private/SettingsEditorModule.cpp
+++ b/Engine/Source/Developer/SettingsEditor/Private/SettingsEditorModule.cpp
@@ -233,7 +233,7 @@ private:
 				if ( Settings->HasAnyFlags(RF_ClassDefaultObject) && !Settings->GetClass()->HasAnyCastFlag(CLASS_Deprecated) )
 				{
 					// Ignore the setting if it's specifically the UDeveloperSettings or other abstract settings classes
-					if ( Settings->GetClass()->HasAnyClassFlags(CLASS_Abstract) )
+					if ( Settings->GetClass()->HasAnyClassFlags(CLASS_Abstract) || !Settings->SupportsAutoRegistration())
 					{
 						continue;
 					}
diff --git a/Engine/Source/Developer/StandaloneRenderer/Private/OpenGL/SlateOpenGLRenderer.cpp b/Engine/Source/Developer/StandaloneRenderer/Private/OpenGL/SlateOpenGLRenderer.cpp
--- a/Engine/Source/Developer/StandaloneRenderer/Private/OpenGL/SlateOpenGLRenderer.cpp
+++ b/Engine/Source/Developer/StandaloneRenderer/Private/OpenGL/SlateOpenGLRenderer.cpp
@@ -282,3 +282,21 @@ ISlateAtlasProvider* FSlateOpenGLRenderer::GetTextureAtlasProvider()
 
 	return nullptr;
 }
+
+int32 FSlateOpenGLRenderer::RegisterCurrentScene(FSceneInterface* Scene) 
+{
+	// This is a no-op
+	return -1;
+}
+
+int32 FSlateOpenGLRenderer::GetCurrentSceneIndex() const
+{
+	// This is a no-op
+	return -1;
+}
+
+
+void FSlateOpenGLRenderer::ClearScenes() 
+{
+	// This is a no-op
+}
\ No newline at end of file
diff --git a/Engine/Source/Developer/StandaloneRenderer/Private/OpenGL/SlateOpenGLRenderer.h b/Engine/Source/Developer/StandaloneRenderer/Private/OpenGL/SlateOpenGLRenderer.h
--- a/Engine/Source/Developer/StandaloneRenderer/Private/OpenGL/SlateOpenGLRenderer.h
+++ b/Engine/Source/Developer/StandaloneRenderer/Private/OpenGL/SlateOpenGLRenderer.h
@@ -163,6 +163,9 @@ public:
 	virtual FSlateUpdatableTexture* CreateUpdatableTexture(uint32 Width, uint32 Height) override;
 	virtual void ReleaseUpdatableTexture(FSlateUpdatableTexture* Texture) override;
 	virtual ISlateAtlasProvider* GetTextureAtlasProvider() override;
+	virtual int32 RegisterCurrentScene(FSceneInterface* Scene) override;
+	virtual int32 GetCurrentSceneIndex() const override;
+	virtual void ClearScenes() override;
 
 private:
 	/** 
diff --git a/Engine/Source/Developer/StandaloneRenderer/Private/Windows/D3D/SlateD3DRenderer.cpp b/Engine/Source/Developer/StandaloneRenderer/Private/Windows/D3D/SlateD3DRenderer.cpp
--- a/Engine/Source/Developer/StandaloneRenderer/Private/Windows/D3D/SlateD3DRenderer.cpp
+++ b/Engine/Source/Developer/StandaloneRenderer/Private/Windows/D3D/SlateD3DRenderer.cpp
@@ -172,6 +172,24 @@ ISlateAtlasProvider* FSlateD3DRenderer::GetTextureAtlasProvider()
 	return nullptr;
 }
 
+int32 FSlateD3DRenderer::RegisterCurrentScene(FSceneInterface* Scene) 
+{
+	// This is a no-op
+	return -1;
+}
+
+int32 FSlateD3DRenderer::GetCurrentSceneIndex() const
+{
+	// This is a no-op
+	return -1;
+}
+
+void FSlateD3DRenderer::ClearScenes() 
+{
+	// This is a no-op
+}
+
+
 void FSlateD3DRenderer::Private_CreateViewport( TSharedRef<SWindow> InWindow, const FVector2D &WindowSize )
 {
 	TSharedRef< FGenericWindow > NativeWindow = InWindow->GetNativeWindow().ToSharedRef();
diff --git a/Engine/Source/Developer/StandaloneRenderer/Private/Windows/D3D/SlateD3DRenderer.h b/Engine/Source/Developer/StandaloneRenderer/Private/Windows/D3D/SlateD3DRenderer.h
--- a/Engine/Source/Developer/StandaloneRenderer/Private/Windows/D3D/SlateD3DRenderer.h
+++ b/Engine/Source/Developer/StandaloneRenderer/Private/Windows/D3D/SlateD3DRenderer.h
@@ -64,6 +64,9 @@ public:
 	virtual FSlateUpdatableTexture* CreateUpdatableTexture(uint32 Width, uint32 Height) override;
 	virtual void ReleaseUpdatableTexture(FSlateUpdatableTexture* Texture) override;
 	virtual ISlateAtlasProvider* GetTextureAtlasProvider() override;
+	virtual int32 RegisterCurrentScene(FSceneInterface* Scene) override;
+	virtual int32 GetCurrentSceneIndex() const override;
+	virtual void ClearScenes() override;
 	
 	bool CreateDevice();
 	void CreateDepthStencilBuffer( FSlateD3DViewport& Viewport );
diff --git a/Engine/Source/Editor/AudioEditor/Private/AssetTypeActions/AssetTypeActions_DialogueWave.cpp b/Engine/Source/Editor/AudioEditor/Private/AssetTypeActions/AssetTypeActions_DialogueWave.cpp
--- a/Engine/Source/Editor/AudioEditor/Private/AssetTypeActions/AssetTypeActions_DialogueWave.cpp
+++ b/Engine/Source/Editor/AudioEditor/Private/AssetTypeActions/AssetTypeActions_DialogueWave.cpp
@@ -19,7 +19,7 @@ void FAssetTypeActions_DialogueWave::GetActions(const TArray<UObject*>& InObject
 	MenuBuilder.AddMenuEntry(
 		LOCTEXT("Sound_PlaySound", "Play"),
 		LOCTEXT("Sound_PlaySoundTooltip", "Plays the selected sound."),
-		FSlateIcon(FEditorStyle::GetStyleSetName(), "MediaAsset.AssetActions.Play"),
+		FSlateIcon(FEditorStyle::GetStyleSetName(), "MediaAsset.AssetActions.Play.Small"),
 		FUIAction(
 		FExecuteAction::CreateSP(this, &FAssetTypeActions_DialogueWave::ExecutePlaySound, DialogueWaves),
 		FCanExecuteAction::CreateSP(this, &FAssetTypeActions_DialogueWave::CanExecutePlayCommand, DialogueWaves)
@@ -29,7 +29,7 @@ void FAssetTypeActions_DialogueWave::GetActions(const TArray<UObject*>& InObject
 	MenuBuilder.AddMenuEntry(
 		LOCTEXT("Sound_StopSound", "Stop"),
 		LOCTEXT("Sound_StopSoundTooltip", "Stops the selected sounds."),
-		FSlateIcon(FEditorStyle::GetStyleSetName(), "MediaAsset.AssetActions.Stop"),
+		FSlateIcon(FEditorStyle::GetStyleSetName(), "MediaAsset.AssetActions.Stop.Small"),
 		FUIAction(
 		FExecuteAction::CreateSP(this, &FAssetTypeActions_DialogueWave::ExecuteStopSound, DialogueWaves),
 		FCanExecuteAction()
diff --git a/Engine/Source/Editor/AudioEditor/Private/AssetTypeActions/AssetTypeActions_SoundBase.cpp b/Engine/Source/Editor/AudioEditor/Private/AssetTypeActions/AssetTypeActions_SoundBase.cpp
--- a/Engine/Source/Editor/AudioEditor/Private/AssetTypeActions/AssetTypeActions_SoundBase.cpp
+++ b/Engine/Source/Editor/AudioEditor/Private/AssetTypeActions/AssetTypeActions_SoundBase.cpp
@@ -24,7 +24,7 @@ void FAssetTypeActions_SoundBase::GetActions( const TArray<UObject*>& InObjects,
 	MenuBuilder.AddMenuEntry(
 		LOCTEXT("Sound_PlaySound", "Play"),
 		LOCTEXT("Sound_PlaySoundTooltip", "Plays the selected sound."),
-		FSlateIcon(FEditorStyle::GetStyleSetName(), "MediaAsset.AssetActions.Play"),
+		FSlateIcon(FEditorStyle::GetStyleSetName(), "MediaAsset.AssetActions.Play.Small"),
 		FUIAction(
 			FExecuteAction::CreateSP( this, &FAssetTypeActions_SoundBase::ExecutePlaySound, Sounds ),
 			FCanExecuteAction::CreateSP( this, &FAssetTypeActions_SoundBase::CanExecutePlayCommand, Sounds )
@@ -34,7 +34,7 @@ void FAssetTypeActions_SoundBase::GetActions( const TArray<UObject*>& InObjects,
 	MenuBuilder.AddMenuEntry(
 		LOCTEXT("Sound_StopSound", "Stop"),
 		LOCTEXT("Sound_StopSoundTooltip", "Stops the selected sounds."),
-		FSlateIcon(FEditorStyle::GetStyleSetName(), "MediaAsset.AssetActions.Stop"),
+		FSlateIcon(FEditorStyle::GetStyleSetName(), "MediaAsset.AssetActions.Stop.Small"),
 		FUIAction(
 			FExecuteAction::CreateSP( this, &FAssetTypeActions_SoundBase::ExecuteStopSound, Sounds ),
 			FCanExecuteAction()
@@ -42,11 +42,6 @@ void FAssetTypeActions_SoundBase::GetActions( const TArray<UObject*>& InObjects,
 		);
 }
 
-bool FAssetTypeActions_SoundBase::CanExecutePlayCommand(TArray<TWeakObjectPtr<USoundBase>> Objects) const
-{
-	return Objects.Num() == 1;
-}
-
 void FAssetTypeActions_SoundBase::AssetsActivated( const TArray<UObject*>& InObjects, EAssetTypeActivationMethod::Type ActivationType )
 {
 	if (ActivationType == EAssetTypeActivationMethod::Previewed)
@@ -88,25 +83,6 @@ void FAssetTypeActions_SoundBase::AssetsActivated( const TArray<UObject*>& InObj
 	}
 }
 
-void FAssetTypeActions_SoundBase::ExecutePlaySound(TArray<TWeakObjectPtr<USoundBase>> Objects) const
-{
-	for (auto ObjIt = Objects.CreateConstIterator(); ObjIt; ++ObjIt)
-	{
-		USoundBase* Sound = (*ObjIt).Get();
-		if ( Sound )
-		{
-			// Only play the first valid sound
-			PlaySound(Sound);
-			break;
-		}
-	}
-}
-
-void FAssetTypeActions_SoundBase::ExecuteStopSound(TArray<TWeakObjectPtr<USoundBase>> Objects) const
-{
-	StopSound();
-}
-
 void FAssetTypeActions_SoundBase::PlaySound(USoundBase* Sound) const
 {
 	if ( Sound )
@@ -124,59 +100,63 @@ void FAssetTypeActions_SoundBase::StopSound() const
 	GEditor->ResetPreviewAudioComponent();
 }
 
-bool FAssetTypeActions_SoundBase::IsSoundPlaying(TArray<TWeakObjectPtr<USoundBase>> Objects) const
+bool FAssetTypeActions_SoundBase::IsSoundPlaying(USoundBase* Sound) const
 {
 	UAudioComponent* PreviewComp = GEditor->GetPreviewAudioComponent();
-	if (PreviewComp)
+	return PreviewComp && PreviewComp->Sound == Sound && PreviewComp->IsPlaying();
+}
+
+void FAssetTypeActions_SoundBase::ExecutePlaySound(TArray<TWeakObjectPtr<USoundBase>> Objects) const
+{
+	for (auto ObjIt = Objects.CreateConstIterator(); ObjIt; ++ObjIt)
 	{
-		for (TWeakObjectPtr<USoundBase> SoundItem : Objects)
+		USoundBase* Sound = (*ObjIt).Get();
+		if (Sound)
 		{
-			if (PreviewComp->Sound == SoundItem && PreviewComp->IsPlaying())
-			{
-				return true;
-			}
+			// Only play the first valid sound
+			PlaySound(Sound);
+			break;
 		}
 	}
+}
 
-	return false;
+void FAssetTypeActions_SoundBase::ExecuteStopSound(TArray<TWeakObjectPtr<USoundBase>> Objects) const
+{
+	StopSound();
 }
 
 TSharedPtr<SWidget> FAssetTypeActions_SoundBase::GetThumbnailOverlay(const FAssetData& AssetData) const
 {
-	TArray<TWeakObjectPtr<USoundBase>> SoundList;
-	SoundList.Add(TWeakObjectPtr<USoundBase>((USoundBase*)AssetData.GetAsset()));
-
-	auto OnGetDisplayBrushLambda = [this, SoundList]() -> const FSlateBrush*
+	auto OnGetDisplayBrushLambda = [this, AssetData]() -> const FSlateBrush*
 	{
-		if (IsSoundPlaying(SoundList))
+		USoundBase* Sound = CastChecked<USoundBase>(AssetData.GetAsset());
+		if (IsSoundPlaying(Sound))
 		{
-			return FEditorStyle::GetBrush("MediaAsset.AssetActions.Stop");
+			return FEditorStyle::GetBrush("MediaAsset.AssetActions.Stop.Large");
 		}
 
-	return FEditorStyle::GetBrush("MediaAsset.AssetActions.Play");
+		return FEditorStyle::GetBrush("MediaAsset.AssetActions.Play.Large");
 	};
 
-	auto IsEnabledLambda = [this, SoundList]() -> bool
+	auto OnClickedLambda = [this, AssetData]() -> FReply
 	{
-		return CanExecutePlayCommand(SoundList);
-	};
+		USoundBase* Sound = CastChecked<USoundBase>(AssetData.GetAsset());
 
-	auto OnClickedLambda = [this, SoundList]() -> FReply
-	{
-		if (IsSoundPlaying(SoundList))
+		if (IsSoundPlaying(Sound))
 		{
-			ExecuteStopSound(SoundList);
+			StopSound();
 		}
 		else
 		{
-			ExecutePlaySound(SoundList);
+			PlaySound(Sound);
 		}
 		return FReply::Handled();
 	};
 
-	auto OnToolTipTextLambda = [this, SoundList]() -> FText
+	auto OnToolTipTextLambda = [this, AssetData]() -> FText
 	{
-		if (IsSoundPlaying(SoundList))
+		USoundBase* Sound = CastChecked<USoundBase>(AssetData.GetAsset());
+		if (IsSoundPlaying(Sound))
 		{
 			return LOCTEXT("Thumbnail_StopSoundToolTip", "Stop selected sound");
 		}
@@ -184,28 +164,49 @@ TSharedPtr<SWidget> FAssetTypeActions_SoundBase::GetThumbnailOverlay(const FAsse
 		return LOCTEXT("Thumbnail_PlaySoundToolTip", "Play selected sound");
 	};
 
-	return SNew(SBox)
+	TSharedPtr<SBox> Box;
+	SAssignNew(Box, SBox)
 		.HAlign(HAlign_Center)
 		.VAlign(VAlign_Center)
-		.Padding(FMargin(2))
+		.Padding(FMargin(2));
+
+	auto OnGetVisibilityLambda = [this, Box, AssetData]() -> EVisibility
+	{
+		USoundBase* Sound = CastChecked<USoundBase>(AssetData.GetAsset());
+		if (Box.IsValid())
+		{
+			if (Box->IsHovered() || IsSoundPlaying(Sound))
+			{
+				return EVisibility::Visible;
+			}
+
+		}
+
+		return EVisibility::Hidden;
+	};
+
+	TSharedPtr<SButton> Widget;
+	SAssignNew(Widget, SButton)
+		.ButtonStyle(FEditorStyle::Get(), "HoverHintOnly")
+		.ToolTipText_Lambda(OnToolTipTextLambda)
+		.Cursor(EMouseCursor::Default) // The outer widget can specify a DragHand cursor, so we need to override that here
+		.ForegroundColor(FSlateColor::UseForeground())
+		.OnClicked_Lambda(OnClickedLambda)
+		.Visibility_Lambda(OnGetVisibilityLambda)
 		[
-			SNew(SButton)
-			.ButtonStyle(FEditorStyle::Get(), "HoverHintOnly")
-			.ToolTipText_Lambda(OnToolTipTextLambda)
-			.Cursor(EMouseCursor::Default) // The outer widget can specify a DragHand cursor, so we need to override that here
-			.ForegroundColor(FSlateColor::UseForeground())		
-			.IsEnabled_Lambda(IsEnabledLambda)
-			.OnClicked_Lambda(OnClickedLambda)
-			[
-				SNew(SBox)
-				.MinDesiredWidth(16)
-				.MinDesiredHeight(16)
-				[
-					SNew(SImage)
-					.Image_Lambda(OnGetDisplayBrushLambda)
-				]
-			]
+			SNew(SImage)
+			.Image_Lambda(OnGetDisplayBrushLambda)
 		];
+
+	Box->SetContent(Widget.ToSharedRef());
+	Box->SetVisibility(EVisibility::Visible);
+
+	return Box;
+}
+
+bool FAssetTypeActions_SoundBase::CanExecutePlayCommand(TArray<TWeakObjectPtr<USoundBase>> Objects) const
+{
+	return Objects.Num() == 1;
 }
 
 #undef LOCTEXT_NAMESPACE
diff --git a/Engine/Source/Editor/AudioEditor/Private/AssetTypeActions/AssetTypeActions_SoundBase.h b/Engine/Source/Editor/AudioEditor/Private/AssetTypeActions/AssetTypeActions_SoundBase.h
--- a/Engine/Source/Editor/AudioEditor/Private/AssetTypeActions/AssetTypeActions_SoundBase.h
+++ b/Engine/Source/Editor/AudioEditor/Private/AssetTypeActions/AssetTypeActions_SoundBase.h
@@ -23,10 +23,17 @@ public:
 	virtual bool CanFilter() override { return false; }
 	virtual TSharedPtr<SWidget> GetThumbnailOverlay(const FAssetData& AssetData) const override;
 
+protected:
+	/** Plays the specified sound wave */
+	void PlaySound(USoundBase* Sound) const;
 
-	/** Return if the specified sound is playing*/
-	bool IsSoundPlaying(TArray<TWeakObjectPtr<USoundBase>> Objects) const;
+	/** Stops any currently playing sounds */
+	void StopSound() const;
 
+	/** Return true if the specified sound is playing */
+	bool IsSoundPlaying(USoundBase* Sound) const;
+
+private:
 	/** Handler for when PlaySound is selected */
 	void ExecutePlaySound(TArray<TWeakObjectPtr<USoundBase>> Objects) const;
 
@@ -35,11 +42,4 @@ public:
 
 	/** Returns true if only one sound is selected to play */
 	bool CanExecutePlayCommand(TArray<TWeakObjectPtr<USoundBase>> Objects) const;
-
-private:
-	/** Plays the specified sound wave */
-	void PlaySound(USoundBase* Sound) const;
-
-	/** Stops any currently playing sounds */
-	void StopSound() const;
 };
diff --git a/Engine/Source/Editor/AudioEditor/Private/AssetTypeActions/AssetTypeActions_SoundWave.cpp b/Engine/Source/Editor/AudioEditor/Private/AssetTypeActions/AssetTypeActions_SoundWave.cpp
--- a/Engine/Source/Editor/AudioEditor/Private/AssetTypeActions/AssetTypeActions_SoundWave.cpp
+++ b/Engine/Source/Editor/AudioEditor/Private/AssetTypeActions/AssetTypeActions_SoundWave.cpp
@@ -202,4 +202,82 @@ void FAssetTypeActions_SoundWave::FillVoiceMenu(FMenuBuilder& MenuBuilder, TArra
 	MenuBuilder.AddWidget(VoicePicker, FText::GetEmpty(), false);
 }
 
+TSharedPtr<SWidget> FAssetTypeActions_SoundWave::GetThumbnailOverlay(const FAssetData& AssetData) const
+{
+	auto OnGetDisplayBrushLambda = [this, AssetData]() -> const FSlateBrush*
+	{
+		USoundBase* Sound = CastChecked<USoundBase>(AssetData.GetAsset());
+		if (IsSoundPlaying(Sound))
+		{
+			return FEditorStyle::GetBrush("MediaAsset.AssetActions.Stop.Large");
+		}
+
+		return FEditorStyle::GetBrush("MediaAsset.AssetActions.Play.Large");
+	};
+
+	auto OnClickedLambda = [this, AssetData]() -> FReply
+	{
+		USoundBase* Sound = CastChecked<USoundBase>(AssetData.GetAsset());
+		if (IsSoundPlaying(Sound))
+		{
+			StopSound();
+		}
+		else
+		{
+			PlaySound(Sound);
+		}
+		return FReply::Handled();
+	};
+
+	auto OnToolTipTextLambda = [this, AssetData]() -> FText
+	{
+		USoundBase* Sound = CastChecked<USoundBase>(AssetData.GetAsset());
+		if (IsSoundPlaying(Sound))
+		{
+			return LOCTEXT("Blueprint_StopSoundToolTip", "Stop selected Sound Wave");
+		}
+
+		return LOCTEXT("Blueprint_PlaySoundToolTip", "Play selected Sound Wave");
+	};
+
+	TSharedPtr<SBox> Box;
+	SAssignNew(Box, SBox)
+		.HAlign(HAlign_Center)
+		.VAlign(VAlign_Center)
+		.Padding(FMargin(2));
+
+	auto OnGetVisibilityLambda = [this, Box, AssetData]() -> EVisibility
+	{
+		USoundBase* Sound = CastChecked<USoundBase>(AssetData.GetAsset());
+		if (Box.IsValid())
+		{
+			if (Box->IsHovered() || IsSoundPlaying(Sound))
+			{
+				return EVisibility::Visible;
+			}
+
+		}
+
+		return EVisibility::Hidden;
+	};
+
+	TSharedPtr<SButton> Widget;
+	SAssignNew(Widget, SButton)
+		.ButtonStyle(FEditorStyle::Get(), "HoverHintOnly")
+		.ToolTipText_Lambda(OnToolTipTextLambda)
+		.Cursor(EMouseCursor::Default) // The outer widget can specify a DragHand cursor, so we need to override that here
+		.ForegroundColor(FSlateColor::UseForeground())
+		.OnClicked_Lambda(OnClickedLambda)
+		.Visibility_Lambda(OnGetVisibilityLambda)
+		[
+			SNew(SImage)
+			.Image_Lambda(OnGetDisplayBrushLambda)
+		];
+
+	Box->SetContent(Widget.ToSharedRef());
+	Box->SetVisibility(EVisibility::Visible);
+
+	return Box;
+}
+
 #undef LOCTEXT_NAMESPACE
diff --git a/Engine/Source/Editor/AudioEditor/Private/AssetTypeActions/AssetTypeActions_SoundWave.h b/Engine/Source/Editor/AudioEditor/Private/AssetTypeActions/AssetTypeActions_SoundWave.h
--- a/Engine/Source/Editor/AudioEditor/Private/AssetTypeActions/AssetTypeActions_SoundWave.h
+++ b/Engine/Source/Editor/AudioEditor/Private/AssetTypeActions/AssetTypeActions_SoundWave.h
@@ -21,7 +21,8 @@ public:
 	virtual bool CanFilter() override { return true; }
 	virtual bool IsImportedAsset() const override { return true; }
 	virtual void GetResolvedSourceFilePaths(const TArray<UObject*>& TypeAssets, TArray<FString>& OutSourceFilePaths) const override;
-	
+	virtual TSharedPtr<SWidget> GetThumbnailOverlay(const FAssetData& AssetData) const override;
+
 private:
 
 	/** Creates a SoundCue of the same name for the sound, if one does not already exist */
diff --git a/Engine/Source/Editor/BlueprintGraph/Classes/EdGraphSchema_K2.h b/Engine/Source/Editor/BlueprintGraph/Classes/EdGraphSchema_K2.h
--- a/Engine/Source/Editor/BlueprintGraph/Classes/EdGraphSchema_K2.h
+++ b/Engine/Source/Editor/BlueprintGraph/Classes/EdGraphSchema_K2.h
@@ -481,6 +481,7 @@ public:
 	virtual bool IsCacheVisualizationOutOfDate(int32 InVisualizationCacheID) const override;
 	virtual int32 GetCurrentVisualizationCacheID() const override;
 	virtual void ForceVisualizationCacheClear() const override;
+	virtual bool SafeDeleteNodeFromGraph(UEdGraph* Graph, UEdGraphNode* NodeToDelete) const override;
 	//~ End EdGraphSchema Interface
 
 	/**
diff --git a/Engine/Source/Editor/BlueprintGraph/Classes/K2Node_Knot.h b/Engine/Source/Editor/BlueprintGraph/Classes/K2Node_Knot.h
--- a/Engine/Source/Editor/BlueprintGraph/Classes/K2Node_Knot.h
+++ b/Engine/Source/Editor/BlueprintGraph/Classes/K2Node_Knot.h
@@ -28,7 +28,8 @@ public:
 	virtual TSharedPtr<class INameValidatorInterface> MakeNameValidator() const override;
 	virtual bool CanSplitPin(const UEdGraphPin* Pin) const override;
 	virtual bool IsCompilerRelevant() const override { return false; }
-	virtual UEdGraphPin* GetPassThroughPin(const UEdGraphPin* FromPin) const override;
+	virtual UEdGraphPin* GetPassThroughPin(const UEdGraphPin* FromPin) const override;	
+	virtual bool ShouldDrawNodeAsControlPointOnly(int32& OutInputPinIndex, int32& OutOutputPinIndex) const override { OutInputPinIndex = 0;  OutOutputPinIndex = 1; return true; }
 	// End of UEdGraphNode interface
 
 	// UK2Node interface
diff --git a/Engine/Source/Editor/BlueprintGraph/Private/EdGraphSchema_K2.cpp b/Engine/Source/Editor/BlueprintGraph/Private/EdGraphSchema_K2.cpp
--- a/Engine/Source/Editor/BlueprintGraph/Private/EdGraphSchema_K2.cpp
+++ b/Engine/Source/Editor/BlueprintGraph/Private/EdGraphSchema_K2.cpp
@@ -6619,6 +6619,23 @@ void UEdGraphSchema_K2::ForceVisualizationCacheClear() const
 	++CurrentCacheRefreshID;
 }
 
+
+bool UEdGraphSchema_K2::SafeDeleteNodeFromGraph(UEdGraph* Graph, UEdGraphNode* NodeToDelete) const 
+{
+	UK2Node* Node = Cast<UK2Node>(NodeToDelete);
+	if (Node == nullptr || Graph == nullptr || NodeToDelete->GetGraph() != Graph)
+	{
+		return false;
+	}
+
+	UBlueprint* OwnerBlueprint = Node->GetBlueprint();
+	Graph->Modify();
+
+	FBlueprintEditorUtils::RemoveNode(OwnerBlueprint, Node, /*bDontRecompile=*/ true);
+	FBlueprintEditorUtils::MarkBlueprintAsModified(OwnerBlueprint);
+	return true;
+}
+
 /////////////////////////////////////////////////////
 
 #undef LOCTEXT_NAMESPACE
diff --git a/Engine/Source/Editor/ClassViewer/ClassViewer.Build.cs b/Engine/Source/Editor/ClassViewer/ClassViewer.Build.cs
--- a/Engine/Source/Editor/ClassViewer/ClassViewer.Build.cs
+++ b/Engine/Source/Editor/ClassViewer/ClassViewer.Build.cs
@@ -27,7 +27,8 @@ public class ClassViewer : ModuleRules
                 "EditorStyle",
 				"UnrealEd",
 				"PropertyEditor",
-			}
+                "Settings",
+            }
 		);
 
 		DynamicallyLoadedModuleNames.AddRange(
diff --git a/Engine/Source/Editor/ClassViewer/Private/ClassViewerModule.cpp b/Engine/Source/Editor/ClassViewer/Private/ClassViewerModule.cpp
--- a/Engine/Source/Editor/ClassViewer/Private/ClassViewerModule.cpp
+++ b/Engine/Source/Editor/ClassViewer/Private/ClassViewerModule.cpp
@@ -11,8 +11,13 @@
 #include "SClassViewer.h"
 #include "Editor/WorkspaceMenuStructure/Public/WorkspaceMenuStructure.h"
 #include "Editor/WorkspaceMenuStructure/Public/WorkspaceMenuStructureModule.h"
-
 #include "Widgets/Docking/SDockTab.h"
+#include "ModuleManager.h"
+#include "SDockTab.h"
+#include "ClassViewerProjectSettings.h"
+#include "ISettingsModule.h"
+
+#define LOCTEXT_NAMESPACE "ClassViewer"
 
 IMPLEMENT_MODULE( FClassViewerModule, ClassViewer );
 
@@ -43,6 +48,18 @@ void FClassViewerModule::StartupModule()
 		.SetTooltipText( NSLOCTEXT("ClassViewerApp", "TooltipText", "Displays all classes that exist within this project.") )
 		.SetGroup( WorkspaceMenu::GetMenuStructure().GetDeveloperToolsMiscCategory() )
 		.SetIcon( FSlateIcon(FEditorStyle::GetStyleSetName(), "ClassViewer.TabIcon") );
+
+	ISettingsModule* SettingsModule = FModuleManager::GetModulePtr<ISettingsModule>("Settings");
+	if (SettingsModule != nullptr)
+	{
+		// ClassViewer Editor Settings
+		SettingsModule->RegisterSettings("Project", "Editor", "ClassViewer",
+			LOCTEXT("ClassViewerSettingsName", "Class Viewer"),
+			LOCTEXT("ClassViewerSettingsDescription", "Configure options for the Class Viewer."),
+			GetMutableDefault<UClassViewerProjectSettings>()
+			);
+	}
+
 }
 
 void FClassViewerModule::ShutdownModule()
@@ -52,6 +69,14 @@ void FClassViewerModule::ShutdownModule()
 		FGlobalTabmanager::Get()->UnregisterNomadTabSpawner( ClassViewerModule::ClassViewerApp );
 	}
 
+	// Unregister the setting
+	ISettingsModule* SettingsModule = FModuleManager::GetModulePtr<ISettingsModule>("Settings");
+
+	if (SettingsModule != nullptr)
+	{
+		SettingsModule->UnregisterSettings("Project", "Editor", "ClassViewer");
+	}
+
 	SClassViewer::DestroyClassHierarchy();
 }
 
@@ -69,3 +94,6 @@ TSharedRef<SWidget> FClassViewerModule::CreateClassViewer(const FClassViewerInit
 	return SNew( SClassViewer, InitOptions )
 			.OnClassPickedDelegate(OnClassPickedDelegate);
 }
+
+
+#undef LOCTEXT_NAMESPACE
\ No newline at end of file
