Copying //UE4/Dev-Editor to //UE4/Dev-Main (Source: //UE4/Dev-Editor @ 3279756)

diff --git a/Engine/Source/Developer/ScreenShotComparison/Private/Widgets/SScreenComparisonRow.cpp b/Engine/Source/Developer/ScreenShotComparison/Private/Widgets/SScreenComparisonRow.cpp
--- a/Engine/Source/Developer/ScreenShotComparison/Private/Widgets/SScreenComparisonRow.cpp
+++ b/Engine/Source/Developer/ScreenShotComparison/Private/Widgets/SScreenComparisonRow.cpp
@@ -13,98 +13,124 @@
 #include "Widgets/Views/SListView.h"
 #include "Widgets/Layout/SScaleBox.h"
 #include "Widgets/Input/SButton.h"
+#include "Widgets/Input/SSlider.h"
 #include "HAL/FileManager.h"
 #include "Misc/FileHelper.h"
 #include "Interfaces/IImageWrapperModule.h"
 #include "Framework/Application/SlateApplication.h"
 #include "JsonObjectConverter.h"
 #include "SScreenShotImagePopup.h"
+#include "SAsyncImage.h"
 
 #define LOCTEXT_NAMESPACE "SScreenShotBrowser"
 
-void SScreenComparisonRow::Construct( const FArguments& InArgs, const TSharedRef<STableViewBase>& InOwnerTableView )
+class SImageComparison : public SCompoundWidget
 {
-	Comparisons = InArgs._Comparisons;
-	ScreenshotManager = InArgs._ScreenshotManager;
-	ComparisonDirectory = InArgs._ComparisonDirectory;
-	Model = InArgs._ComparisonResult;
+public:
+	SLATE_BEGIN_ARGS(SImageComparison) {}
+		SLATE_ARGUMENT(TSharedPtr<FSlateDynamicImageBrush>, BaseImage)
+		SLATE_ARGUMENT(TSharedPtr<FSlateDynamicImageBrush>, ModifiedImage)
+	SLATE_END_ARGS()
 
-	CachedActualImageSize = FIntPoint::NoneValue;
+	void Construct(const FArguments& InArgs)
+	{
+		BaseImage = InArgs._BaseImage;
+		ModifiedImage = InArgs._ModifiedImage;
+
+		//const int32 MaxHeight = FSlateApplication::Get().GetPreferredWorkArea().GetSize().Y * 0.80;
+
+		ChildSlot
+		[
+			SNew(SVerticalBox)
+
+			+ SVerticalBox::Slot()
+			.FillHeight(1.0f)
+			[
+				SNew(SScaleBox)
+				.Stretch(EStretch::ScaleToFit)
+				[
+					SNew(SOverlay)
 
-	switch ( Model->GetType() )
+					+ SOverlay::Slot()
+					[
+						SNew(SImage)
+						.Image(BaseImage.Get())
+					]
+						
+					+ SOverlay::Slot()
+					[
+						SNew(SImage)
+						.Image(ModifiedImage.Get())
+						.ColorAndOpacity(this, &SImageComparison::GetModifiedOpacity)
+					]
+				]
+			]
+			
+			+ SVerticalBox::Slot()
+			.AutoHeight()
+			.HAlign(HAlign_Center)
+			[
+				SNew(SHorizontalBox)
+				
+				+ SHorizontalBox::Slot()
+				.AutoWidth()
+				[
+					SNew(STextBlock)
+					.Text(LOCTEXT("GroundTruth", "Ground Truth"))
+				]
+
+				+SHorizontalBox::Slot()
+				.AutoWidth()
+				[
+					SNew(SBox)
+					.MinDesiredWidth(150)
+					[
+						SAssignNew(OpacitySlider, SSlider)
+						.Value(0.5f)
+					]
+				]
+
+				+ SHorizontalBox::Slot()
+				.AutoWidth()
+				[
+					SNew(STextBlock)
+					.Text(LOCTEXT("Incoming", "Incoming"))
+				]
+			]
+		];
+	}
+
+	FSlateColor GetModifiedOpacity() const
 	{
-		case EComparisonResultType::Added:
-		{
-			FString AddedFolder = ComparisonDirectory / Comparisons->IncomingPath / Model->Folder;
-			IFileManager::Get().FindFilesRecursive(ExternalFiles, *AddedFolder, TEXT("*.*"), true, false);
-			LoadMetadata();
-			break;
-		}
-		case EComparisonResultType::Missing:
-		{
-			break;
-		}
-		case EComparisonResultType::Compared:
-		{
-			const FImageComparisonResult& ComparisonResult = Model->ComparisonResult.GetValue();
-			FString IncomingImage = ComparisonDirectory / Comparisons->IncomingPath / ComparisonResult.IncomingFile;
-			FString IncomingMetadata = FPaths::ChangeExtension(IncomingImage, TEXT("json"));
-			LoadMetadata();
-
-			ExternalFiles.Add(IncomingImage);
-			ExternalFiles.Add(IncomingMetadata);
-			break;
-		}
+		return FLinearColor(1, 1, 1, OpacitySlider->GetValue());
 	}
 
-	SMultiColumnTableRow<TSharedPtr<FScreenComparisonModel>>::Construct(FSuperRowType::FArguments(), InOwnerTableView);
-}
+private:
+	TSharedPtr<FSlateDynamicImageBrush> BaseImage;
+	TSharedPtr<FSlateDynamicImageBrush> ModifiedImage;
+
+	TSharedPtr<SSlider> OpacitySlider;
+};
+
 
-void SScreenComparisonRow::LoadMetadata()
+void SScreenComparisonRow::Construct( const FArguments& InArgs, const TSharedRef<STableViewBase>& InOwnerTableView )
 {
-	if ( !Model->Metadata.IsSet() )
-	{
-		FString IncomingMetadataFile;
+	ScreenshotManager = InArgs._ScreenshotManager;
+	ComparisonDirectory = InArgs._ComparisonDirectory;
+	Model = InArgs._ComparisonResult;
 
-		if ( Model->ComparisonResult.IsSet() )
-		{
-			const FImageComparisonResult& ComparisonResult = Model->ComparisonResult.GetValue();
-			FString IncomingImage = ComparisonDirectory / Comparisons->IncomingPath / ComparisonResult.IncomingFile;
-			IncomingMetadataFile = FPaths::ChangeExtension(IncomingImage, TEXT("json"));
-		}
-		else
-		{
-			TArray<FString> MetadataFiles;
-			FString AddedFolder = ComparisonDirectory / Comparisons->IncomingPath / Model->Folder;
-			IFileManager::Get().FindFilesRecursive(MetadataFiles, *AddedFolder, TEXT("*.json"), true, false);
-			if ( MetadataFiles.Num() > 0 )
-			{
-				IncomingMetadataFile = MetadataFiles[0];
-			}
-		}
+	CachedActualImageSize = FIntPoint::NoneValue;
 
-		if ( !IncomingMetadataFile.IsEmpty() )
-		{
-			FString Json;
-			if ( FFileHelper::LoadFileToString(Json, *IncomingMetadataFile) )
-			{
-				FAutomationScreenshotMetadata Metadata;
-				if ( FJsonObjectConverter::JsonObjectStringToUStruct(Json, &Metadata, 0, 0) )
-				{
-					Model->Metadata = Metadata;
-				}
-			}
-		}
-	}
+	SMultiColumnTableRow<TSharedPtr<FScreenComparisonModel>>::Construct(FSuperRowType::FArguments(), InOwnerTableView);
 }
 
 TSharedRef<SWidget> SScreenComparisonRow::GenerateWidgetForColumn(const FName& ColumnName)
 {
 	if ( ColumnName == "Name" )
 	{
-		if ( Model->Metadata.IsSet() )
+		if ( Model->GetMetadata().IsSet() )
 		{
-			return SNew(STextBlock).Text(FText::FromString(Model->Metadata->Name));
+			return SNew(STextBlock).Text(FText::FromString(Model->GetMetadata()->Name));
 		}
 		else
 		{
@@ -113,69 +139,58 @@ TSharedRef<SWidget> SScreenComparisonRow::GenerateWidgetForColumn(const FName& C
 	}
 	else if ( ColumnName == "Delta" )
 	{
-		if ( Model->ComparisonResult.IsSet() )
-		{
-			FNumberFormattingOptions Format;
-			Format.MinimumFractionalDigits = 2;
-			Format.MaximumFractionalDigits = 2;
-			const FText GlobalDelta = FText::AsPercent(Model->ComparisonResult->GlobalDifference, &Format);
-			const FText LocalDelta = FText::AsPercent(Model->ComparisonResult->MaxLocalDifference, &Format);
-
-			const FText Differences = FText::Format(LOCTEXT("LocalvGlobalDelta", "{0} | {1}"), LocalDelta, GlobalDelta);
-			return SNew(STextBlock).Text(Differences);
-		}
-
-		return SNew(STextBlock).Text(FText::AsPercent(1.0));
+		FNumberFormattingOptions Format;
+		Format.MinimumFractionalDigits = 2;
+		Format.MaximumFractionalDigits = 2;
+		const FText GlobalDelta = FText::AsPercent(Model->Report.Comparison.GlobalDifference, &Format);
+		const FText LocalDelta = FText::AsPercent(Model->Report.Comparison.MaxLocalDifference, &Format);
+
+		const FText Differences = FText::Format(LOCTEXT("LocalvGlobalDelta", "{0} | {1}"), LocalDelta, GlobalDelta);
+		return SNew(STextBlock).Text(Differences);
 	}
 	else if ( ColumnName == "Preview" )
 	{
-		switch ( Model->GetType() )
+		const FImageComparisonResult& ComparisonResult = Model->Report.Comparison;
+		if ( ComparisonResult.IsNew() )
 		{
-			case EComparisonResultType::Added:
-			{
-				return BuildAddedView();
-			}
-			case EComparisonResultType::Missing:
-			{
-				return BuildMissingView();
-			}
-			case EComparisonResultType::Compared:
-			{
-				return
-					SNew(SVerticalBox)
-
-					+ SVerticalBox::Slot()
-					.AutoHeight()
+			return BuildAddedView();
+		}
+		else
+		{
+			return
+				SNew(SVerticalBox)
+
+				+ SVerticalBox::Slot()
+				.AutoHeight()
+				[
+					BuildComparisonPreview()
+				]
+
+				+ SVerticalBox::Slot()
+				.AutoHeight()
+				.HAlign(HAlign_Center)
+				[
+					SNew(SHorizontalBox)
+
+					+ SHorizontalBox::Slot()
+					.AutoWidth()
 					[
-						BuildComparisonPreview()
+						SNew(SButton)
+						.IsEnabled(this, &SScreenComparisonRow::CanReplace)
+						.Text(LOCTEXT("Replace", "Replace"))
+						.OnClicked(this, &SScreenComparisonRow::Replace)
 					]
 
-					+ SVerticalBox::Slot()
-					.AutoHeight()
-					.HAlign(HAlign_Center)
+					+ SHorizontalBox::Slot()
+					.AutoWidth()
+					.Padding(10, 0, 0, 0)
 					[
-						SNew(SHorizontalBox)
-
-						+ SHorizontalBox::Slot()
-						.AutoWidth()
-						[
-							SNew(SButton)
-							.IsEnabled(this, &SScreenComparisonRow::CanReplace)
-							.Text(LOCTEXT("Replace", "Replace"))
-							.OnClicked(this, &SScreenComparisonRow::Replace)
-						]
-
-						+ SHorizontalBox::Slot()
-						.AutoWidth()
-						.Padding(10, 0, 0, 0)
-						[
-							SNew(SButton)
-							.IsEnabled(this, &SScreenComparisonRow::CanAddAsAlternative)
-							.Text(LOCTEXT("AddAlternative", "Add As Alternative"))
-							.OnClicked(this, &SScreenComparisonRow::AddAlternative)
-						]
-					];
-			}
+						SNew(SButton)
+						.IsEnabled(this, &SScreenComparisonRow::CanAddAsAlternative)
+						.Text(LOCTEXT("AddAlternative", "Add As Alternative"))
+						.OnClicked(this, &SScreenComparisonRow::AddAlternative)
+					]
+				];
 		}
 	}
 
@@ -187,32 +202,10 @@ bool SScreenComparisonRow::CanUseSourceControl() const
 	return ISourceControlModule::Get().IsEnabled();
 }
 
-TSharedRef<SWidget> SScreenComparisonRow::BuildMissingView()
-{
-	return
-		SNew(SHorizontalBox)
-		
-		/*	+SHorizontalBox::Slot()
-			[
-				SNew(SButton)
-				.IsEnabled(this, &SScreenComparisonRow::CanUseSourceControl)
-				.Text(LOCTEXT("RemoveOld", "Remove Old"))
-				.OnClicked(this, &SScreenComparisonRow::RemoveOld)
-			]*/;
-}
-
 TSharedRef<SWidget> SScreenComparisonRow::BuildAddedView()
 {
-	TArray<FString> Images;
-
-	FString AddedFolder = ComparisonDirectory / Comparisons->IncomingPath / Model->Folder;
-	IFileManager::Get().FindFilesRecursive(Images, *AddedFolder, TEXT("*.png"), true, false);
-
-	//TODO Automation this is no good, we need to clear out the incoming directory beforehand.
-	if ( Images.Num() > 0 )
-	{
-		UnapprovedBrush = LoadScreenshot(Images[0]);
-	}
+	const FImageComparisonResult& ComparisonResult = Model->Report.Comparison;
+	FString IncomingFile = Model->Report.ReportFolder / ComparisonResult.ReportIncomingFile;
 
 	return
 		SNew(SVerticalBox)
@@ -221,7 +214,7 @@ TSharedRef<SWidget> SScreenComparisonRow::BuildAddedView()
 		.AutoHeight()
 		[
 			SNew(SBox)
-			.MaxDesiredHeight(100)
+			.HeightOverride(100)
 			.HAlign(HAlign_Left)
 			[
 				SNew(SScaleBox)
@@ -233,9 +226,9 @@ TSharedRef<SWidget> SScreenComparisonRow::BuildAddedView()
 					.AutoWidth()
 					.Padding(4.0f, 4.0f)
 					[
-						SNew(SImage)
-						.Image(UnapprovedBrush.Get())
-						.OnMouseButtonDown(this, &SScreenComparisonRow::OnImageClicked, UnapprovedBrush)
+						SNew(SAsyncImage)
+						.ImageFilePath(IncomingFile)
+						//.OnMouseButtonDown(this, &SScreenComparisonRow::OnImageClicked, UnapprovedBrush)
 					]
 				]
 			]
@@ -254,66 +247,90 @@ TSharedRef<SWidget> SScreenComparisonRow::BuildAddedView()
 
 TSharedRef<SWidget> SScreenComparisonRow::BuildComparisonPreview()
 {
-	const FImageComparisonResult& ComparisonResult = Model->ComparisonResult.GetValue();
-
-	FString ApprovedFile = ComparisonDirectory / Comparisons->ApprovedPath / ComparisonResult.ApprovedFile;
-	FString IncomingFile = ComparisonDirectory / Comparisons->IncomingPath / ComparisonResult.IncomingFile;
-	FString DeltaFile = ComparisonDirectory / Comparisons->DeltaPath / ComparisonResult.ComparisonFile;
+	const FImageComparisonResult& ComparisonResult = Model->Report.Comparison;
 
-	ApprovedBrush = LoadScreenshot(ApprovedFile);
-	UnapprovedBrush = LoadScreenshot(IncomingFile);
-	ComparisonBrush = LoadScreenshot(DeltaFile);
+	FString ApprovedFile = Model->Report.ReportFolder / ComparisonResult.ReportApprovedFile;
+	FString IncomingFile = Model->Report.ReportFolder / ComparisonResult.ReportIncomingFile;
+	FString DeltaFile = Model->Report.ReportFolder / ComparisonResult.ReportComparisonFile;
 
 	// Create the screen shot data widget.
 	return 
-		SNew(SBox)
-		.MaxDesiredHeight(100)
-		.HAlign(HAlign_Left)
+		SNew(SBorder)
+		.BorderImage(nullptr)
+		.OnMouseButtonDown(this, &SScreenComparisonRow::OnCompareImages)
 		[
-			SNew(SScaleBox)
-			.Stretch(EStretch::ScaleToFit)
+			SNew(SVerticalBox)
+
+			+ SVerticalBox::Slot()
+			.AutoHeight()
+			[
+				SNew(SBox)
+				.HeightOverride(100)
+				[
+					SNew(SScaleBox)
+					.Stretch(EStretch::ScaleToFit)
+					[
+						SNew(SHorizontalBox)
+					
+						+ SHorizontalBox::Slot()
+						.AutoWidth()
+						.Padding(4.0f, 4.0f)
+						[
+							SAssignNew(ApprovedImageWidget, SAsyncImage)
+							.ImageFilePath(ApprovedFile)
+						]
+
+						+ SHorizontalBox::Slot()
+						.AutoWidth()
+						.Padding(4.0f, 4.0f)
+						[
+							SNew(SAsyncImage)
+							.ImageFilePath(DeltaFile)
+						]
+
+						+ SHorizontalBox::Slot()
+						.AutoWidth()
+						.Padding(4.0f, 4.0f)
+						[
+							SAssignNew(UnapprovedImageWidget, SAsyncImage)
+							.ImageFilePath(IncomingFile)
+						]
+					]
+				]
+			]
+	
+			+ SVerticalBox::Slot()
+			.AutoHeight()
 			[
 				SNew(SHorizontalBox)
-		
+
 				+ SHorizontalBox::Slot()
-				.AutoWidth()
-				.Padding(4.0f, 4.0f)
+				.FillWidth(1.0f)
+				.HAlign(HAlign_Center)
 				[
-					SNew(SImage)
-					.Image(ApprovedBrush.Get())
-					.OnMouseButtonDown(this, &SScreenComparisonRow::OnImageClicked, ApprovedBrush)
+					SNew(STextBlock)
+					.Text(LOCTEXT( "GroundTruth", "Ground Truth" ))
 				]
 
 				+ SHorizontalBox::Slot()
-				.AutoWidth()
-				.Padding(4.0f, 4.0f)
+				.FillWidth(1.0f)
+				.HAlign(HAlign_Center)
 				[
-					SNew(SImage)
-					.Image(ComparisonBrush.Get())
-					.OnMouseButtonDown(this, &SScreenComparisonRow::OnImageClicked, ComparisonBrush)
+					SNew(STextBlock)
+					.Text(LOCTEXT("Difference", "Difference"))
 				]
 
 				+ SHorizontalBox::Slot()
-				.AutoWidth()
-				.Padding(4.0f, 4.0f)
+				.FillWidth(1.0f)
+				.HAlign(HAlign_Center)
 				[
-					SNew(SImage)
-					.Image(UnapprovedBrush.Get())
-					.OnMouseButtonDown(this, &SScreenComparisonRow::OnImageClicked, UnapprovedBrush)
+					SNew(STextBlock)
+					.Text(LOCTEXT("Incoming", "Incoming"))
 				]
 			]
 		];
 }
 
-void SScreenComparisonRow::GetStatus()
-{
-	//if ( SourceControlStates.Num() == 0 )
-	//{
-	//	ISourceControlProvider& SourceControlProvider = ISourceControlModule::Get().GetProvider();
-	//	SourceControlProvider.GetState(SourceControlFiles, SourceControlStates, EStateCacheUsage::Use);
-	//}
-}
-
 bool SScreenComparisonRow::CanAddNew() const
 {
 	return CanUseSourceControl();
@@ -321,30 +338,7 @@ bool SScreenComparisonRow::CanAddNew() const
 
 FReply SScreenComparisonRow::AddNew()
 {
-	// Copy files to the approved
-	const FString& LocalApprovedFolder = ScreenshotManager->GetLocalApprovedFolder();
-	const FString ImportIncomingRoot = ComparisonDirectory / Comparisons->IncomingPath / TEXT("");
-
-	TArray<FString> SourceControlFiles;
-
-	for ( const FString& File : ExternalFiles )
-	{
-		FString RelativeFile = File;
-		FPaths::MakePathRelativeTo(RelativeFile, *ImportIncomingRoot);
-		
-		FString DestFilePath = LocalApprovedFolder / RelativeFile;
-		IFileManager::Get().Copy(*DestFilePath, *File, true, true);
-
-		SourceControlFiles.Add(DestFilePath);
-	}
-
-	ISourceControlProvider& SourceControlProvider = ISourceControlModule::Get().GetProvider();
-	if ( SourceControlProvider.Execute(ISourceControlOperation::Create<FMarkForAdd>(), SourceControlFiles) == ECommandResult::Failed )
-	{
-		//TODO Error
-	}
-
-	// Invalidate Status
+	Model->AddNew(ScreenshotManager);
 
 	return FReply::Handled();
 }
@@ -354,185 +348,63 @@ bool SScreenComparisonRow::CanReplace() const
 	return CanUseSourceControl();
 }
 
-FReply SScreenComparisonRow::RemoveExistingApproved()
-{
-	TArray<FString> FilesToRemove;
-
-	FString PlatformFolder = ComparisonDirectory / Comparisons->ApprovedPath / FPaths::GetPath(Model->ComparisonResult->ApprovedFile);
-	IFileManager::Get().FindFilesRecursive(FilesToRemove, *PlatformFolder, TEXT("*.*"), true, false);
-
-	// Copy files to the approved
-	const FString& LocalApprovedFolder = ScreenshotManager->GetLocalApprovedFolder();
-	const FString ImportApprovedRoot = ComparisonDirectory / Comparisons->ApprovedPath / TEXT("");
-
-	TArray<FString> SourceControlFiles;
-
-	for ( const FString& File : FilesToRemove )
-	{
-		FString RelativeFile = File;
-		FPaths::MakePathRelativeTo(RelativeFile, *ImportApprovedRoot);
-
-		FString DestFilePath = LocalApprovedFolder / RelativeFile;
-		SourceControlFiles.Add(DestFilePath);
-
-		IFileManager::Get().Delete(*DestFilePath, false, true, false);
-	}
-
-	ISourceControlProvider& SourceControlProvider = ISourceControlModule::Get().GetProvider();
-	if ( SourceControlProvider.Execute(ISourceControlOperation::Create<FDelete>(), SourceControlFiles) == ECommandResult::Failed )
-	{
-		//TODO Error
-	}
-
-	// Invalidate Status
-
-	return FReply::Handled();
-}
-
 FReply SScreenComparisonRow::Replace()
 {
-	// Delete all the existing files in this area
-	RemoveExistingApproved();
-
-	// Copy files to the approved
-	const FString& LocalApprovedFolder = ScreenshotManager->GetLocalApprovedFolder();
-	const FString ImportIncomingRoot = ComparisonDirectory / Comparisons->IncomingPath / TEXT("");
-
-	TArray<FString> SourceControlFiles;
-
-	for ( const FString& File : ExternalFiles )
-	{
-		FString RelativeFile = File;
-		FPaths::MakePathRelativeTo(RelativeFile, *ImportIncomingRoot);
-
-		FString DestFilePath = LocalApprovedFolder / RelativeFile;
-		
-		SourceControlFiles.Add(DestFilePath);
-	}
-
-	ISourceControlProvider& SourceControlProvider = ISourceControlModule::Get().GetProvider();
-	if ( SourceControlProvider.Execute(ISourceControlOperation::Create<FRevert>(), SourceControlFiles) == ECommandResult::Failed )
-	{
-		//TODO Error
-	}
-
-	for ( const FString& File : ExternalFiles )
-	{
-		FString RelativeFile = File;
-		FPaths::MakePathRelativeTo(RelativeFile, *ImportIncomingRoot);
-
-		FString DestFilePath = LocalApprovedFolder / RelativeFile;
-		IFileManager::Get().Copy(*DestFilePath, *File, true, true);
-	}
-
-	if ( SourceControlProvider.Execute(ISourceControlOperation::Create<FMarkForAdd>(), SourceControlFiles) == ECommandResult::Failed )
-	{
-		//TODO Error
-	}
-	if ( SourceControlProvider.Execute(ISourceControlOperation::Create<FCheckOut>(), SourceControlFiles) == ECommandResult::Failed )
-	{
-		//TODO Error
-	}
+	Model->Replace(ScreenshotManager);
 
 	return FReply::Handled();
 }
 
 bool SScreenComparisonRow::CanAddAsAlternative() const
 {
-	return CanUseSourceControl() && (Model->ComparisonResult->IncomingFile != Model->ComparisonResult->ApprovedFile);
+	return CanUseSourceControl() && (Model->Report.Comparison.IncomingFile != Model->Report.Comparison.ApprovedFile);
 }
 
 FReply SScreenComparisonRow::AddAlternative()
 {
-	// Copy files to the approved
-	const FString& LocalApprovedFolder = ScreenshotManager->GetLocalApprovedFolder();
-	const FString ImportIncomingRoot = ComparisonDirectory / Comparisons->IncomingPath / TEXT("");
-
-	TArray<FString> SourceControlFiles;
-
-	for ( const FString& File : ExternalFiles )
-	{
-		FString RelativeFile = File;
-		FPaths::MakePathRelativeTo(RelativeFile, *ImportIncomingRoot);
-
-		FString DestFilePath = LocalApprovedFolder / RelativeFile;
-
-		SourceControlFiles.Add(DestFilePath);
-	}
-
-	ISourceControlProvider& SourceControlProvider = ISourceControlModule::Get().GetProvider();
-	if ( SourceControlProvider.Execute(ISourceControlOperation::Create<FRevert>(), SourceControlFiles) == ECommandResult::Failed )
-	{
-		//TODO Error
-	}
-
-	for ( const FString& File : ExternalFiles )
-	{
-		FString RelativeFile = File;
-		FPaths::MakePathRelativeTo(RelativeFile, *ImportIncomingRoot);
-
-		FString DestFilePath = LocalApprovedFolder / RelativeFile;
-		IFileManager::Get().Copy(*DestFilePath, *File, true, true);
-	}
-
-	if ( SourceControlProvider.Execute(ISourceControlOperation::Create<FMarkForAdd>(), SourceControlFiles) == ECommandResult::Failed )
-	{
-		//TODO Error
-	}
-	if ( SourceControlProvider.Execute(ISourceControlOperation::Create<FCheckOut>(), SourceControlFiles) == ECommandResult::Failed )
-	{
-		//TODO Error
-	}
+	Model->AddAlternative(ScreenshotManager);
 
 	return FReply::Handled();
 }
 
-//TSharedRef<SWidget> SScreenComparisonRow::BuildSourceControl()
-//{
-//	if ( ISourceControlModule::Get().IsEnabled() )
-//	{
-//		// note: calling QueueStatusUpdate often does not spam status updates as an internal timer prevents this
-//		ISourceControlModule::Get().QueueStatusUpdate(CachedConfigFileName);
-//
-//		ISourceControlProvider& SourceControlProvider = ISourceControlModule::Get().GetProvider();
-//	}
-//}
-
-TSharedPtr<FSlateDynamicImageBrush> SScreenComparisonRow::LoadScreenshot(FString ImagePath)
+FReply SScreenComparisonRow::OnCompareImages(const FGeometry& InGeometry, const FPointerEvent& InEvent)
 {
-	TArray<uint8> RawFileData;
-	if ( FFileHelper::LoadFileToArray(RawFileData, *ImagePath) )
+	TSharedPtr<FSlateDynamicImageBrush> ApprovedImage = ApprovedImageWidget->GetDynamicBrush();
+	TSharedPtr<FSlateDynamicImageBrush> UnapprovedImage = UnapprovedImageWidget->GetDynamicBrush();
+
+	if ( ApprovedImage.IsValid() && UnapprovedImage.IsValid() )
 	{
-		IImageWrapperModule& ImageWrapperModule = FModuleManager::LoadModuleChecked<IImageWrapperModule>(FName("ImageWrapper"));
-		IImageWrapperPtr ImageWrapper = ImageWrapperModule.CreateImageWrapper(EImageFormat::PNG);
+		TSharedRef<SWindow> ParentWindow = FSlateApplication::Get().FindWidgetWindow(AsShared()).ToSharedRef();
+
+		// Center ourselves in the parent window
+		TSharedRef<SWindow> PopupWindow = SNew(SWindow)
+			.IsPopupWindow(false)
+			.ClientSize(FVector2D(1280, 720))
+			.SizingRule(ESizingRule::UserSized)
+			.AutoCenter(EAutoCenter::PreferredWorkArea)
+			.SupportsMaximize(true)
+			.SupportsMinimize(true)
+			.FocusWhenFirstShown(true)
+			.ActivateWhenFirstShown(true)
+			.Content()
+			[
+				SNew(SImageComparison)
+				.BaseImage(ApprovedImage)
+				.ModifiedImage(UnapprovedImage)
+			];
 
-		if ( ImageWrapper.IsValid() && ImageWrapper->SetCompressed(RawFileData.GetData(), RawFileData.Num()) )
-		{
-			const TArray<uint8>* RawData = NULL;
-			if ( ImageWrapper->GetRaw(ERGBFormat::BGRA, 8, RawData) )
-			{
-				if ( FSlateApplication::Get().GetRenderer()->GenerateDynamicImageResource(*ImagePath, ImageWrapper->GetWidth(), ImageWrapper->GetHeight(), *RawData) )
-				{
-					return MakeShareable(new FSlateDynamicImageBrush(*ImagePath, FVector2D(ImageWrapper->GetWidth(), ImageWrapper->GetHeight())));
-				}
-			}
-		}
+		FSlateApplication::Get().AddWindowAsNativeChild(PopupWindow, ParentWindow, true);
 	}
 
-	return nullptr;
+	return FReply::Handled();
 }
 
 FReply SScreenComparisonRow::OnImageClicked(const FGeometry& InGeometry, const FPointerEvent& InEvent, TSharedPtr<FSlateDynamicImageBrush> Image)
 {
-	TSharedRef<SScreenShotImagePopup> PopupImage =
-		SNew(SScreenShotImagePopup)
-		.ImageBrush(Image)
-		.ImageSize(Image->ImageSize.IntPoint());
-
-	auto ParentWindow = FSlateApplication::Get().FindWidgetWindow(AsShared()).ToSharedRef();
+	TSharedRef<SWindow> ParentWindow = FSlateApplication::Get().FindWidgetWindow(AsShared()).ToSharedRef();
 
 	// Center ourselves in the parent window
-	auto PopupWindow = SNew(SWindow)
+	TSharedRef<SWindow> PopupWindow = SNew(SWindow)
 		.IsPopupWindow(false)
 		.SizingRule(ESizingRule::UserSized)
 		.ClientSize(Image->ImageSize)
@@ -543,7 +415,9 @@ FReply SScreenComparisonRow::OnImageClicked(const FGeometry& InGeometry, const F
 		.ActivateWhenFirstShown(true)
 		.Content()
 		[
-			PopupImage
+			SNew(SScreenShotImagePopup)
+			.ImageBrush(Image)
+			.ImageSize(Image->ImageSize.IntPoint())
 		];
 
 	FSlateApplication::Get().AddWindowAsNativeChild(PopupWindow, ParentWindow, true);
diff --git a/Engine/Source/Developer/ScreenShotComparison/Private/Widgets/SScreenComparisonRow.h b/Engine/Source/Developer/ScreenShotComparison/Private/Widgets/SScreenComparisonRow.h
--- a/Engine/Source/Developer/ScreenShotComparison/Private/Widgets/SScreenComparisonRow.h
+++ b/Engine/Source/Developer/ScreenShotComparison/Private/Widgets/SScreenComparisonRow.h
@@ -10,9 +10,11 @@
 #include "Widgets/Views/STableRow.h"
 #include "ISourceControlState.h"
 #include "Interfaces/IScreenShotManager.h"
+#include "Async/Async.h"
 
 class FScreenComparisonModel;
 struct FSlateDynamicImageBrush;
+class SAsyncImage;
 
 /**
  * Widget to display a particular view.
@@ -26,7 +28,6 @@ public:
 
 		SLATE_ARGUMENT( IScreenShotManagerPtr, ScreenshotManager )
 		SLATE_ARGUMENT( FString, ComparisonDirectory )
-		SLATE_ARGUMENT( TSharedPtr<FComparisonResults>, Comparisons )
 		SLATE_ARGUMENT( TSharedPtr<FScreenComparisonModel>, ComparisonResult )
 
 	SLATE_END_ARGS()
@@ -44,7 +45,6 @@ public:
 private:
 	bool CanUseSourceControl() const;
 
-	TSharedRef<SWidget> BuildMissingView();
 	TSharedRef<SWidget> BuildAddedView();
 	TSharedRef<SWidget> BuildComparisonPreview();
 
@@ -57,13 +57,7 @@ private:
 	bool CanAddAsAlternative() const;
 	FReply AddAlternative();
 
-	FReply RemoveExistingApproved();
-
-	void GetStatus();
-
-	TSharedPtr<FSlateDynamicImageBrush> LoadScreenshot(FString ImagePath);
-	void LoadMetadata();
-
+	FReply OnCompareImages(const FGeometry& InGeometry, const FPointerEvent& InEvent);
 	FReply OnImageClicked(const FGeometry& InGeometry, const FPointerEvent& InEvent, TSharedPtr<FSlateDynamicImageBrush> Image);
 
 private:
@@ -76,23 +70,9 @@ private:
 
 	FString ComparisonDirectory;
 
-	TSharedPtr<FComparisonResults> Comparisons;
-
 	//The cached actual size of the screenshot
 	FIntPoint CachedActualImageSize;
 
-	//Holds the dynamic brush.
-	TSharedPtr<FSlateDynamicImageBrush> ApprovedBrush;
-
-	//Holds the dynamic brush.
-	TSharedPtr<FSlateDynamicImageBrush> UnapprovedBrush;
-
-	//Holds the dynamic brush.
-	TSharedPtr<FSlateDynamicImageBrush> ComparisonBrush;
-
-	// 
-	TArray<FString> ExternalFiles;
-
-	// 
-	TArray<FSourceControlStateRef> SourceControlStates;
+	TSharedPtr<SAsyncImage> ApprovedImageWidget;
+	TSharedPtr<SAsyncImage> UnapprovedImageWidget;
 };
diff --git a/Engine/Source/Developer/ScreenShotComparison/Private/Widgets/SScreenShotBrowser.cpp b/Engine/Source/Developer/ScreenShotComparison/Private/Widgets/SScreenShotBrowser.cpp
--- a/Engine/Source/Developer/ScreenShotComparison/Private/Widgets/SScreenShotBrowser.cpp
+++ b/Engine/Source/Developer/ScreenShotComparison/Private/Widgets/SScreenShotBrowser.cpp
@@ -12,20 +12,27 @@
 #include "Widgets/SScreenComparisonRow.h"
 #include "Models/ScreenComparisonModel.h"
 #include "Misc/FeedbackContext.h"
+#include "EditorStyleSet.h"
+#include "Modules/ModuleManager.h"
+#include "DirectoryWatcherModule.h"
+#include "IDirectoryWatcher.h"
 
 #define LOCTEXT_NAMESPACE "ScreenshotComparison"
 
 void SScreenShotBrowser::Construct( const FArguments& InArgs,  IScreenShotManagerRef InScreenShotManager  )
 {
 	ScreenShotManager = InScreenShotManager;
-	ComparisonRoot = FPaths::ConvertRelativePathToFull(FPaths::GameSavedDir() / TEXT("Exported"));
+	ComparisonRoot = FPaths::ConvertRelativePathToFull(FPaths::GameSavedDir() / TEXT("Automation/Comparisons"));
+	bReportsChanged = true;
 
+	FModuleManager::Get().LoadModuleChecked(FName("ImageWrapper"));
 
 	ChildSlot
 	[
 		SNew( SVerticalBox )
-		+SVerticalBox::Slot()
+		+ SVerticalBox::Slot()
 		.AutoHeight()
+		.Padding(2.0f)
 		[
 			SNew(SHorizontalBox)
 
@@ -38,26 +45,29 @@ void SScreenShotBrowser::Construct( const FArguments& InArgs,  IScreenShotManage
 			]
 
 			+ SHorizontalBox::Slot()
-			.AutoWidth()
+			.FillWidth(1.0f)
+			.HAlign(HAlign_Right)
 			[
 				SNew(SButton)
-				.Text(LOCTEXT("RunComparison", "Run Compare"))
-				.OnClicked_Lambda(
-					[&]()
+				.HAlign(HAlign_Center)
+				.Text(LOCTEXT("DeleteAllReports", "Delete All Reports"))
+				.ToolTipText(LOCTEXT("DeleteAllReportsTooltip", "Deletes all the current reports.  Reports are not removed unless the user resolves them, \nso if you just want to reset the state of the reports, clear them here and then re-run the tests."))
+				.ButtonStyle(FEditorStyle::Get(), "FlatButton.Danger")
+				.TextStyle(FEditorStyle::Get(), "FlatButton.DefaultTextStyle")
+				.OnClicked_Lambda([this]()
 				{
-					GWarn->BeginSlowTask(LOCTEXT("ComparingScreenshots", "Comparing Screenshots..."), true);
-					ScreenShotManager->CompareScreensotsAsync().Wait();
-					ScreenShotManager->ExportScreensotsAsync().Wait();
-					GWarn->EndSlowTask();
-
-					RebuildTree();
+					while ( ComparisonList.Num() > 0 )
+					{
+						TSharedPtr<FScreenComparisonModel> Model = ComparisonList.Pop();
+						Model->Complete();
+					}
 
 					return FReply::Handled();
 				})
 			]
 		]
-
-		+SVerticalBox::Slot()
+		
+		+ SVerticalBox::Slot()
 		.FillHeight( 1.0f )
 		[
 			SAssignNew(ComparisonView, SListView< TSharedPtr<FScreenComparisonModel> >)
@@ -91,17 +101,49 @@ void SScreenShotBrowser::Construct( const FArguments& InArgs,  IScreenShotManage
 		]
 	];
 
-	// Register for callbacks
-	ScreenShotManager->RegisterScreenShotUpdate(ScreenShotDelegate);
+	RefreshDirectoryWatcher();
+}
 
-	RebuildTree();
+SScreenShotBrowser::~SScreenShotBrowser()
+{
+	FDirectoryWatcherModule& DirectoryWatcherModule = FModuleManager::LoadModuleChecked<FDirectoryWatcherModule>(TEXT("DirectoryWatcher"));
+
+	if ( DirectoryWatchingHandle.IsValid() )
+	{
+		DirectoryWatcherModule.Get()->UnregisterDirectoryChangedCallback_Handle(ComparisonRoot, DirectoryWatchingHandle);
+	}
+}
+
+void SScreenShotBrowser::Tick(const FGeometry& AllottedGeometry, const double InCurrentTime, const float InDeltaTime)
+{
+	if ( bReportsChanged )
+	{
+		RebuildTree();
+	}
 }
 
 void SScreenShotBrowser::OnDirectoryChanged(const FString& Directory)
 {
 	ComparisonRoot = Directory;
 
-	RebuildTree();
+	RefreshDirectoryWatcher();
+}
+
+void SScreenShotBrowser::RefreshDirectoryWatcher()
+{
+	FDirectoryWatcherModule& DirectoryWatcherModule = FModuleManager::LoadModuleChecked<FDirectoryWatcherModule>(TEXT("DirectoryWatcher"));
+	
+	if ( DirectoryWatchingHandle.IsValid() )
+	{
+		DirectoryWatcherModule.Get()->UnregisterDirectoryChangedCallback_Handle(ComparisonRoot, DirectoryWatchingHandle);
+	}
+
+	DirectoryWatcherModule.Get()->RegisterDirectoryChangedCallback_Handle(ComparisonRoot, IDirectoryWatcher::FDirectoryChanged::CreateRaw(this, &SScreenShotBrowser::OnReportsChanged), DirectoryWatchingHandle, IDirectoryWatcher::WatchOptions::IncludeDirectoryChanges);
+}
+
+void SScreenShotBrowser::OnReportsChanged(const TArray<struct FFileChangeData>& /*FileChanges*/)
+{
+	bReportsChanged = true;
 }
 
 TSharedRef<ITableRow> SScreenShotBrowser::OnGenerateWidgetForScreenResults(TSharedPtr<FScreenComparisonModel> InItem, const TSharedRef<STableViewBase>& OwnerTable)
@@ -110,52 +152,26 @@ TSharedRef<ITableRow> SScreenShotBrowser::OnGenerateWidgetForScreenResults(TShar
 	return
 		SNew(SScreenComparisonRow, OwnerTable)
 		.ScreenshotManager(ScreenShotManager)
-		.ComparisonDirectory(ComparisonDirectory)
-		.Comparisons(CurrentComparisons)
+		.ComparisonDirectory(ComparisonRoot)
 		.ComparisonResult(InItem);
 }
 
 void SScreenShotBrowser::RebuildTree()
 {
-	TArray<FString> Changelists;
-	FString SearchDirectory = ComparisonRoot / TEXT("*");
-	IFileManager::Get().FindFiles(Changelists, *SearchDirectory, false, true);
-
+	bReportsChanged = false;
 	ComparisonList.Reset();
 
-	if ( Changelists.Num() > 0 )
+	if ( ScreenShotManager->OpenComparisonReports(ComparisonRoot, CurrentReports) )
 	{
-		Changelists.StableSort();
-
-		ComparisonDirectory = ComparisonRoot / Changelists.Top();
-		CurrentComparisons = ScreenShotManager->ImportScreensots(ComparisonDirectory);
-
-		if ( CurrentComparisons.IsValid() )
+		for ( const FComparisonReport& Report : CurrentReports )
 		{
-			// Copy the comparisons to an array as shared pointers the list view can use.
-			for ( FString& Added : CurrentComparisons->Added )
-			{
-				TSharedPtr<FScreenComparisonModel> Model = MakeShared<FScreenComparisonModel>(EComparisonResultType::Added);
-				Model->Folder = Added;
-				ComparisonList.Add(Model);
-			}
-
-			// Ignore Missing
-			//for ( FString& Missing : CurrentComparisons->Missing )
-			//{
-			//	TSharedPtr<FScreenComparisonModel> Model = MakeShared<FScreenComparisonModel>(EComparisonResultType::Missing);
-			//	Model->Folder = Missing;
-			//	ComparisonList.Add(Model);
-			//}
-
-			// Copy the comparisons to an array as shared pointers the list view can use.
-			for ( FImageComparisonResult& Result : CurrentComparisons->Comparisons )
-			{
-				TSharedPtr<FScreenComparisonModel> Model = MakeShared<FScreenComparisonModel>(EComparisonResultType::Compared);
-				Model->ComparisonResult = Result;
-
-				ComparisonList.Add(Model);
-			}
+			TSharedPtr<FScreenComparisonModel> Model = MakeShared<FScreenComparisonModel>(Report);
+			Model->OnComplete.AddLambda([this, Model] () {
+				ComparisonList.Remove(Model);
+				ComparisonView->RequestListRefresh();
+			});
+
+			ComparisonList.Add(Model);
 		}
 	}
 
diff --git a/Engine/Source/Developer/ScreenShotComparison/Private/Widgets/SScreenShotBrowser.h b/Engine/Source/Developer/ScreenShotComparison/Private/Widgets/SScreenShotBrowser.h
--- a/Engine/Source/Developer/ScreenShotComparison/Private/Widgets/SScreenShotBrowser.h
+++ b/Engine/Source/Developer/ScreenShotComparison/Private/Widgets/SScreenShotBrowser.h
@@ -19,8 +19,7 @@ class FScreenComparisonModel;
  * Implements a Slate widget for browsing active game sessions.
  */
 
-class SScreenShotBrowser
-	: public SCompoundWidget
+class SScreenShotBrowser : public SCompoundWidget
 {
 public:
 
@@ -36,6 +35,9 @@ public:
  	 * @param InScreenShotManager - The screen shot manager containing the screen shot data.
 	 */
 	void Construct( const FArguments& InArgs, IScreenShotManagerRef InScreenShotManager );
+	virtual ~SScreenShotBrowser();
+
+	virtual void Tick(const FGeometry& AllottedGeometry, const double InCurrentTime, const float InDeltaTime) override;
 
 public:
 
@@ -45,6 +47,10 @@ private:
 
 	void OnDirectoryChanged(const FString& Directory);
 
+	void RefreshDirectoryWatcher();
+
+	void OnReportsChanged(const TArray<struct FFileChangeData>& /*FileChanges*/);
+
 	/**
 	 * Regenerate the widgets when the filter changes
 	 */
@@ -58,11 +64,8 @@ private:
 	/** The directory where we're imported comparisons from. */
 	FString ComparisonRoot;
 
-	/** The directory where we're imported comparisons from, with changelist */
-	FString ComparisonDirectory;
-
 	/** The imported screenshot results */
-	TSharedPtr<FComparisonResults> CurrentComparisons;
+	TArray<FComparisonReport> CurrentReports;
 
 	/** The imported screenshot results copied into an array usable by the list view */
 	TArray<TSharedPtr<FScreenComparisonModel>> ComparisonList;
@@ -70,6 +73,9 @@ private:
 	/**  */
 	TSharedPtr< SListView< TSharedPtr<FScreenComparisonModel> > > ComparisonView;
 
-	// Delegate to call when screen shot data changes 
-	FOnScreenFilterChanged ScreenShotDelegate;
+	/** Directory watching handle */
+	FDelegateHandle DirectoryWatchingHandle;
+
+	/**  */
+	bool bReportsChanged;
 };
diff --git a/Engine/Source/Developer/ScreenShotComparison/ScreenShotComparison.Build.cs b/Engine/Source/Developer/ScreenShotComparison/ScreenShotComparison.Build.cs
--- a/Engine/Source/Developer/ScreenShotComparison/ScreenShotComparison.Build.cs
+++ b/Engine/Source/Developer/ScreenShotComparison/ScreenShotComparison.Build.cs
@@ -25,7 +25,8 @@ public class ScreenShotComparison : ModuleRules
 				"SourceControl",
 				"AutomationMessages",
 				"Json",
-				"JsonUtilities"
+				"JsonUtilities",
+				"DirectoryWatcher"
 			}
 		);
 
diff --git a/Engine/Source/Developer/ScreenShotComparisonTools/Private/ImageComparer.cpp b/Engine/Source/Developer/ScreenShotComparisonTools/Private/ImageComparer.cpp
--- a/Engine/Source/Developer/ScreenShotComparisonTools/Private/ImageComparer.cpp
+++ b/Engine/Source/Developer/ScreenShotComparisonTools/Private/ImageComparer.cpp
@@ -82,7 +82,7 @@ public:
 		FString TempDir = OutputDirectory.IsEmpty() ? FString(FPlatformProcess::UserTempDir()) : OutputDirectory;
 		FString TempDeltaFile = FPaths::CreateTempFilename(*TempDir, TEXT("ImageCompare-"), TEXT(".png"));
 
-		IImageWrapperModule& ImageWrapperModule = FModuleManager::LoadModuleChecked<IImageWrapperModule>(FName("ImageWrapper"));
+		IImageWrapperModule& ImageWrapperModule = FModuleManager::GetModuleChecked<IImageWrapperModule>(FName("ImageWrapper"));
 		IImageWrapperPtr ImageWriter = ImageWrapperModule.CreateImageWrapper(EImageFormat::PNG);
 
 		if ( ImageWriter.IsValid() )
@@ -195,6 +195,13 @@ bool FPixelOperations::IsAntialiased(const FColor& SourcePixel, FComparableImage
 	return false;
 }
 
+FComparisonReport::FComparisonReport(const FString& InReportRootDirectory, const FString& InReportFile)
+{
+	ReportRootDirectory = InReportRootDirectory;
+	ReportFile = InReportFile;
+	ReportFolder = FPaths::GetPath(InReportFile);
+}
+
 void FComparableImage::Process()
 {
 	ParallelFor(Width,
@@ -224,7 +231,7 @@ void FComparableImage::Process()
 
 TSharedPtr<FComparableImage> FImageComparer::Open(const FString& ImagePath, FText& OutError)
 {
-	IImageWrapperModule& ImageWrapperModule = FModuleManager::LoadModuleChecked<IImageWrapperModule>(FName("ImageWrapper"));
+	IImageWrapperModule& ImageWrapperModule = FModuleManager::GetModuleChecked<IImageWrapperModule>(FName("ImageWrapper"));
 	IImageWrapperPtr ImageReader = ImageWrapperModule.CreateImageWrapper(EImageFormat::PNG);
 
 	if ( !ImageReader.IsValid() )
@@ -313,7 +320,10 @@ FImageComparisonResult FImageComparer::Compare(const FString& ImagePathA, const
 
 	if ( ImageA->Width != ImageB->Width || ImageA->Height != ImageB->Height )
 	{
-		Results.ErrorMessage = LOCTEXT("DifferentSizesUnsupported", "We can not compare images of different sizes at this time.");
+		//Results.ErrorMessage = LOCTEXT("DifferentSizesUnsupported", "We can not compare images of different sizes at this time.");
+		Results.Tolerance = Tolerance;
+		Results.MaxLocalDifference = 1.0f;
+		Results.GlobalDifference = 1.0f;
 		return Results;
 	}
 
diff --git a/Engine/Source/Developer/ScreenShotComparisonTools/Private/ScreenShotManager.cpp b/Engine/Source/Developer/ScreenShotComparisonTools/Private/ScreenShotManager.cpp
--- a/Engine/Source/Developer/ScreenShotComparisonTools/Private/ScreenShotManager.cpp
+++ b/Engine/Source/Developer/ScreenShotComparisonTools/Private/ScreenShotManager.cpp
@@ -2,14 +2,14 @@
 
 #include "ScreenShotManager.h"
 #include "HAL/FileManager.h"
+#include "Async/Async.h"
 #include "Misc/FileHelper.h"
 #include "Misc/Paths.h"
-#include "Async/Async.h"
 #include "Misc/EngineVersion.h"
 #include "Misc/FilterCollection.h"
 #include "AutomationWorkerMessages.h"
 #include "Helpers/MessageEndpointBuilder.h"
-#include "ScreenShotBaseNode.h"
+#include "Modules/ModuleManager.h"
 
 #include "Serialization/JsonReader.h"
 #include "Serialization/JsonSerializer.h"
@@ -28,13 +28,20 @@ public:
 
 FScreenShotManager::FScreenShotManager()
 {
+	FModuleManager::Get().LoadModuleChecked(FName("ImageWrapper"));
+
 	ScreenshotUnapprovedFolder = FPaths::GameSavedDir() / TEXT("Automation/Incoming/");
 	ScreenshotDeltaFolder = FPaths::GameSavedDir() / TEXT("Automation/Delta/");
 	ScreenshotResultsFolder = FPaths::GameSavedDir() / TEXT("Automation/");
 	ScreenshotApprovedFolder = FPaths::GameDir() / TEXT("Test/Screenshots/");
 
+	ComparisonResultsFolder = FPaths::GameSavedDir() / TEXT("Automation/Comparisons");
+
 	// Clear the incoming directory when we initialize, we don't care about previous runs.
 	//IFileManager::Get().DeleteDirectory(*ScreenshotUnapprovedFolder, false, true);
+
+	// Clear previous comparison results from the local comparison folder.
+	//IFileManager::Get().DeleteDirectory(*ComparisonResultsRoot, false, true);
 }
 
 FString FScreenShotManager::GetLocalUnapprovedFolder() const
@@ -47,173 +54,14 @@ FString FScreenShotManager::GetLocalApprovedFolder() const
 	return FPaths::ConvertRelativePathToFull(ScreenshotApprovedFolder);
 }
 
-TSharedRef<FScreenshotComparisons> FScreenShotManager::GenerateFileList() const
-{
-	TSharedRef<FScreenshotComparisons> Comparisons = MakeShareable(new FScreenshotComparisons());
-
-	Comparisons->ApprovedFolder = ScreenshotApprovedFolder;
-	Comparisons->UnapprovedFolder = ScreenshotUnapprovedFolder;
-
-	TArray<FString> UnapprovedShots;
-	TArray<FString> ApprovedShots;
-
-	{
-		TArray<FString> UnapprovedShotList;
-		IFileManager::Get().FindFilesRecursive(UnapprovedShotList, *ScreenshotUnapprovedFolder, TEXT("*.png"), true, false);
-
-		TArray<FString> ApprovedShotList;
-		IFileManager::Get().FindFilesRecursive(ApprovedShotList, *ScreenshotApprovedFolder, TEXT("*.png"), true, false);
-
-		for ( FString UnapprovedPngFile : UnapprovedShotList )
-		{
-			FPaths::MakePathRelativeTo(UnapprovedPngFile, *ScreenshotUnapprovedFolder);
-			UnapprovedShots.AddUnique(FPaths::GetPath(UnapprovedPngFile));
-		}
-
-		for ( FString ApprovedPngFile : ApprovedShotList )
-		{
-			FPaths::MakePathRelativeTo(ApprovedPngFile, *ScreenshotApprovedFolder);
-			ApprovedShots.AddUnique(FPaths::GetPath(ApprovedPngFile));
-		}
-	}
-
-	for ( int32 ShotIndex = 0; ShotIndex < UnapprovedShots.Num(); ShotIndex++ )
-	{
-		const FString& UnapprovedShot = UnapprovedShots[ShotIndex];
-
-		// Look for an exact match first
-		if ( ApprovedShots.Contains(UnapprovedShot) )
-		{
-			Comparisons->Existing.Add(UnapprovedShot);
-			ApprovedShots.Remove(UnapprovedShot);
-		}
-		else
-		{
-			Comparisons->New.Add(UnapprovedShot);
-		}
-
-		UnapprovedShots.RemoveAt(ShotIndex);
-		ShotIndex--;
-	}
-
-	for ( int32 ShotIndex = 0; ShotIndex < ApprovedShots.Num(); ShotIndex++ )
-	{
-		const FString& ApprovedShot = ApprovedShots[ShotIndex];
-
-		Comparisons->Missing.Add(ApprovedShot);
-	}
-
-	return Comparisons;
-}
-
-void FScreenShotManager::CreateData()
-{
-	TSet<FString> TempPlatforms;
-
-	// Clear existing data
-	CachedPlatformList.Empty();
-	ScreenShotDataArray.Empty();
-
-	// Add platforms for searching
-	CachedPlatformList.Add( MakeShareable(new FString("Any") ) );
-
-	for ( FString Platform : TempPlatforms )
-	{
-		TSharedPtr<FString> PlatformName = MakeShareable(new FString(Platform));
-		CachedPlatformList.Add(PlatformName);
-	}
-
-	SetFilter( nullptr );
-}
-
-TArray< TSharedPtr<FScreenShotDataItem> >& FScreenShotManager::GetScreenshotResults()
-{
-	return ScreenShotDataArray;
-}
-
-
-/* IScreenShotManager interface
- *****************************************************************************/
-
-void FScreenShotManager::GenerateLists()
-{
-	// Create the screen shot data
-	CreateData();
-}
-
-TArray<TSharedPtr<FString> >& FScreenShotManager::GetCachedPlatfomList()
-{
-	return CachedPlatformList;
-}
-
-void FScreenShotManager::RegisterScreenShotUpdate(const FOnScreenFilterChanged& InDelegate )
-{
-	// Set the callback delegate
-	ScreenFilterChangedDelegate = InDelegate;
-}
-
-void FScreenShotManager::SetFilter( TSharedPtr< ScreenShotFilterCollection > InFilter )
-{
-	// Update the UI
-	ScreenFilterChangedDelegate.ExecuteIfBound();
-}
-
 /* IScreenShotManager event handlers
  *****************************************************************************/
 
-TFuture<void> FScreenShotManager::CompareScreensotsAsync()
-{
-	return Async<void>(EAsyncExecution::Thread, [&] () { CompareScreensots(); });
-}
-
 TFuture<FImageComparisonResult> FScreenShotManager::CompareScreensotAsync(FString RelativeImagePath)
 {
 	return Async<FImageComparisonResult>(EAsyncExecution::Thread, [&] () { return CompareScreensot(RelativeImagePath); });
 }
 
-void FScreenShotManager::CompareScreensots()
-{
-	TSharedRef<FScreenshotComparisons> Comparisons = GenerateFileList();
-
-	IFileManager::Get().DeleteDirectory(*ScreenshotDeltaFolder, /*RequireExists =*/false, /*Tree =*/true);
-
-	FComparisonResults Results;
-	Results.ApprovedPath = TEXT("Approved");
-	Results.IncomingPath = TEXT("Incoming");
-	Results.DeltaPath = TEXT("Delta");
-
-	for ( const FString& New : Comparisons->New )
-	{
-		Results.Added.Add(New);
-	}
-
-	for ( const FString& Missing : Comparisons->Missing )
-	{
-		Results.Missing.Add(Missing);
-	}
-
-	for ( const FString& ExistingGroup : Comparisons->Existing )
-	{
-		FString TestUnapprovedFolder = FPaths::Combine(ScreenshotUnapprovedFolder, ExistingGroup);
-
-		TArray<FString> UnapprovedDeviceShots;
-		IFileManager::Get().FindFilesRecursive(UnapprovedDeviceShots, *TestUnapprovedFolder, TEXT("*.png"), true, false);
-
-		check(UnapprovedDeviceShots.Num() > 0);
-
-		for ( const FString& UnapprovedShot : UnapprovedDeviceShots )
-		{
-			FString RelativeUnapprovedShot = UnapprovedShot;
-			FPaths::MakePathRelativeTo(RelativeUnapprovedShot, *ScreenshotUnapprovedFolder);
-
-			const FImageComparisonResult ShotResult = CompareScreensot(RelativeUnapprovedShot);
-			Results.Comparisons.Add(ShotResult);
-		}
-	}
-
-	SaveComparisonResults(Results);
-}
-
 FImageComparisonResult FScreenShotManager::CompareScreensot(FString ExistingImage)
 {
 	FString Existing = FPaths::GetPath(ExistingImage);
@@ -233,91 +81,135 @@ FImageComparisonResult FScreenShotManager::CompareScreensot(FString ExistingImag
 	TArray<FString> ApprovedDeviceShots;
 	IFileManager::Get().FindFilesRecursive(ApprovedDeviceShots, *TestApprovedFolder, TEXT("*.png"), true, false);
 
-	if ( ApprovedDeviceShots.Num() == 0 )
-	{
-		return FImageComparisonResult();
-	}
-
-	// Load the metadata for the incoming unapproved image.
-	FString UnapprovedFileName = FPaths::GetCleanFilename(ExistingImage);
-	FString UnapprovedFullPath = FPaths::Combine(TestUnapprovedFolder, UnapprovedFileName);
-
-	TOptional<FAutomationScreenshotMetadata> ExistingMetadata;
+	FImageComparisonResult ComparisonResult;
 
+	// We can't find a ground truth, so it's a new comparison.
+	if ( ApprovedDeviceShots.Num() > 0 )
 	{
-		// Always read the metadata file from the unapproved location, as we may have introduced new comparison rules.
-		FString MetadataFile = FPaths::ChangeExtension(UnapprovedFullPath, ".json");
+		// Load the metadata for the incoming unapproved image.
+		FString UnapprovedFileName = FPaths::GetCleanFilename(ExistingImage);
+		FString UnapprovedFullPath = FPaths::Combine(TestUnapprovedFolder, UnapprovedFileName);
+
+		TOptional<FAutomationScreenshotMetadata> ExistingMetadata;
 
-		FString Json;
-		if ( FFileHelper::LoadFileToString(Json, *MetadataFile) )
 		{
-			FAutomationScreenshotMetadata Metadata;
-			if ( FJsonObjectConverter::JsonObjectStringToUStruct(Json, &Metadata, 0, 0) )
+			// Always read the metadata file from the unapproved location, as we may have introduced new comparison rules.
+			FString MetadataFile = FPaths::ChangeExtension(UnapprovedFullPath, ".json");
+
+			FString Json;
+			if ( FFileHelper::LoadFileToString(Json, *MetadataFile) )
 			{
-				ExistingMetadata = Metadata;
+				FAutomationScreenshotMetadata Metadata;
+				if ( FJsonObjectConverter::JsonObjectStringToUStruct(Json, &Metadata, 0, 0) )
+				{
+					ExistingMetadata = Metadata;
+				}
 			}
 		}
-	}
 
-	FString NearestExistingApprovedImage;
-	TOptional<FAutomationScreenshotMetadata> NearestExistingApprovedImageMetadata;
-	int32 MatchScore = 0;
+		FString NearestExistingApprovedImage;
+		TOptional<FAutomationScreenshotMetadata> NearestExistingApprovedImageMetadata;
 
-	if ( ExistingMetadata.IsSet() )
-	{
-		for ( FString ApprovedShot : ApprovedDeviceShots )
+		if ( ExistingMetadata.IsSet() )
 		{
-			FString ApprovedShotFile = FPaths::GetCleanFilename(ApprovedShot);
-			FString ApprovedShotFileFull = FPaths::Combine(TestApprovedFolder, ApprovedShotFile);
+			int32 MatchScore = -1;
 
-			FString ApprovedMetadataFile = FPaths::ChangeExtension(ApprovedShotFileFull, ".json");
-
-			FString Json;
-			if ( FFileHelper::LoadFileToString(Json, *ApprovedMetadataFile) )
+			for ( FString ApprovedShot : ApprovedDeviceShots )
 			{
-				FAutomationScreenshotMetadata Metadata;
-				if ( FJsonObjectConverter::JsonObjectStringToUStruct(Json, &Metadata, 0, 0) )
+				FString ApprovedShotFile = FPaths::GetCleanFilename(ApprovedShot);
+				FString ApprovedShotFileFull = FPaths::Combine(TestApprovedFolder, ApprovedShotFile);
+
+				FString ApprovedMetadataFile = FPaths::ChangeExtension(ApprovedShotFileFull, ".json");
+
+				FString Json;
+				if ( FFileHelper::LoadFileToString(Json, *ApprovedMetadataFile) )
 				{
-					int32 Comparison = Metadata.Compare(ExistingMetadata.GetValue());
-					if ( Comparison > MatchScore )
+					FAutomationScreenshotMetadata Metadata;
+					if ( FJsonObjectConverter::JsonObjectStringToUStruct(Json, &Metadata, 0, 0) )
 					{
-						MatchScore = Comparison;
-						NearestExistingApprovedImage = ApprovedShotFile;
-						NearestExistingApprovedImageMetadata = Metadata;
+						int32 Comparison = Metadata.Compare(ExistingMetadata.GetValue());
+						if ( Comparison > MatchScore )
+						{
+							MatchScore = Comparison;
+							NearestExistingApprovedImage = ApprovedShotFile;
+							NearestExistingApprovedImageMetadata = Metadata;
+						}
 					}
 				}
 			}
 		}
+		else
+		{
+			// TODO no metadata how do I pick a good shot?
+			NearestExistingApprovedImage = FPaths::GetCleanFilename(ApprovedDeviceShots[0]);
+		}
+
+		FString ApprovedFullPath = FPaths::Combine(TestApprovedFolder, NearestExistingApprovedImage);
+
+		FImageTolerance Tolerance = DefaultTolerance;
+
+		if ( ExistingMetadata.IsSet() && ExistingMetadata->bHasComparisonRules )
+		{
+			Tolerance.Red = ExistingMetadata->ToleranceRed;
+			Tolerance.Green = ExistingMetadata->ToleranceGreen;
+			Tolerance.Blue = ExistingMetadata->ToleranceBlue;
+			Tolerance.Alpha = ExistingMetadata->ToleranceAlpha;
+			Tolerance.MinBrightness = ExistingMetadata->ToleranceMinBrightness;
+			Tolerance.MaxBrightness = ExistingMetadata->ToleranceMaxBrightness;
+			Tolerance.IgnoreAntiAliasing = ExistingMetadata->bIgnoreAntiAliasing;
+			Tolerance.IgnoreColors = ExistingMetadata->bIgnoreColors;
+			Tolerance.MaximumLocalError = ExistingMetadata->MaximumLocalError;
+			Tolerance.MaximumGlobalError = ExistingMetadata->MaximumGlobalError;
+		}
+
+		// TODO Think about using SSIM, but needs local SSIM as well as Global SSIM, same as the basic comparison.
+		//double SSIM = Comparer.CompareStructuralSimilarity(ApprovedFullPath, UnapprovedFullPath, FImageComparer::EStructuralSimilarityComponent::Luminance);
+		//printf("%f\n", SSIM);
+
+		ComparisonResult = Comparer.Compare(ApprovedFullPath, UnapprovedFullPath, Tolerance);
 	}
 	else
 	{
-		// TODO no metadata how do I pick a good shot?
-		NearestExistingApprovedImage = FPaths::GetCleanFilename(ApprovedDeviceShots[0]);
+		ComparisonResult.IncomingFile = ExistingImage;
 	}
 
-	FString ApprovedFullPath = FPaths::Combine(TestApprovedFolder, NearestExistingApprovedImage);
+	// Generate and save a report of the comparison if it's new or the results are not similar
+	if ( ComparisonResult.IsNew() || !ComparisonResult.AreSimilar() )
+	{
+		FString ReportFolder = ComparisonResultsFolder / Existing;
 
-	FImageTolerance Tolerance = DefaultTolerance;
+		FString ApprovedFile = ( ScreenshotApprovedFolder / ComparisonResult.ApprovedFile );
+		FString ApprovedMetadataFile = FPaths::ChangeExtension(ApprovedFile, ".json");
 
-	if ( ExistingMetadata.IsSet() && ExistingMetadata->bHasComparisonRules )
-	{
-		Tolerance.Red = ExistingMetadata->ToleranceRed;
-		Tolerance.Green = ExistingMetadata->ToleranceGreen;
-		Tolerance.Blue = ExistingMetadata->ToleranceBlue;
-		Tolerance.Alpha = ExistingMetadata->ToleranceAlpha;
-		Tolerance.MinBrightness = ExistingMetadata->ToleranceMinBrightness;
-		Tolerance.MaxBrightness = ExistingMetadata->ToleranceMaxBrightness;
-		Tolerance.IgnoreAntiAliasing = ExistingMetadata->bIgnoreAntiAliasing;
-		Tolerance.IgnoreColors = ExistingMetadata->bIgnoreColors;
-		Tolerance.MaximumLocalError = ExistingMetadata->MaximumLocalError;
-		Tolerance.MaximumGlobalError = ExistingMetadata->MaximumGlobalError;
-	}
+		FString IncomingFile = ( ScreenshotUnapprovedFolder / ComparisonResult.IncomingFile );
+		FString IncomingMetadataFile = FPaths::ChangeExtension(IncomingFile, ".json");
+
+		if ( IFileManager::Get().Copy(*( ReportFolder / TEXT("Approved.png")), *ApprovedFile, true, true) == COPY_OK )
+		{
+			IFileManager::Get().Copy(*( ReportFolder / TEXT("Approved.json")), *ApprovedMetadataFile, true, true);
+			ComparisonResult.ReportApprovedFile = TEXT("Approved.png");
+		}
+
+		if ( IFileManager::Get().Copy(*( ReportFolder / TEXT("Incoming.png")), *IncomingFile, true, true) == COPY_OK )
+		{
+			IFileManager::Get().Copy(*( ReportFolder / TEXT("Incoming.json")), *IncomingMetadataFile, true, true);
+			ComparisonResult.ReportIncomingFile = TEXT("Incoming.png");
+		}
 
-	// TODO Think about using SSIM, but needs local SSIM as well as Global SSIM, same as the basic comparison.
-	//double SSIM = Comparer.CompareStructuralSimilarity(ApprovedFullPath, UnapprovedFullPath, FImageComparer::EStructuralSimilarityComponent::Luminance);
-	//printf("%f\n", SSIM);
+		if ( IFileManager::Get().Copy(*( ReportFolder / TEXT("Delta.png")), *( ScreenshotDeltaFolder / ComparisonResult.ComparisonFile ), true, true) == COPY_OK )
+		{
+			ComparisonResult.ReportComparisonFile = TEXT("Delta.png");
+		}
+
+		FString Json;
+		if ( FJsonObjectConverter::UStructToJsonObjectString(ComparisonResult, Json) )
+		{
+			FString ComparisonReportFile = ReportFolder / TEXT("Report.json");
+			FFileHelper::SaveStringToFile(Json, *ComparisonReportFile, FFileHelper::EEncodingOptions::ForceUTF8WithoutBOM);
+		}
+	}
 
-	return Comparer.Compare(ApprovedFullPath, UnapprovedFullPath, Tolerance);
+	return ComparisonResult;
 }
 
 TFuture<void> FScreenShotManager::ExportScreensotsAsync(FString ExportPath)
@@ -334,8 +226,9 @@ bool FScreenShotManager::ExportComparisonResults(FString RootExportFolder)
 		RootExportFolder = GetDefaultExportDirectory();
 	}
 
-	IFileManager::Get().DeleteDirectory(*RootExportFolder, /*RequireExists =*/false, /*Tree =*/true);
-	if ( !IFileManager::Get().MakeDirectory(*RootExportFolder, /*Tree =*/true) )
+	FString ExportDirectory = RootExportFolder / FString::FromInt(FEngineVersion::Current().GetChangelist());
+
+	if ( !IFileManager::Get().MakeDirectory(*ExportDirectory, /*Tree =*/true) )
 	{
 		return false;
 	}
@@ -343,56 +236,46 @@ bool FScreenShotManager::ExportComparisonResults(FString RootExportFolder)
 	// Wait for file operations to complete.
 	FPlatformProcess::Sleep(1.0f);
 
-	FString ExportDirectory = RootExportFolder / FString::FromInt(FEngineVersion::Current().GetChangelist());
-	FString ExportDirectory_ComparisonManifest = ExportDirectory / TEXT("ComparisonResults.json");
-	FString ExportDirectory_Approved = ExportDirectory / TEXT("Approved");
-	FString ExportDirectory_Incoming = ExportDirectory / TEXT("Incoming");
-	FString ExportDirectory_Delta = ExportDirectory / TEXT("Delta");
-
-	// Copy the comparison manifest
-	IFileManager::Get().Copy(*ExportDirectory_ComparisonManifest, *GetComparisonResultsJsonFilename(), true, true);
-
-	TArray<FString> FilesToCopy;
-
-	// Copy the ground truth images
-	CopyDirectory(ExportDirectory_Approved, ScreenshotApprovedFolder);
-
-	// Copy the incoming unapproved images
-	CopyDirectory(ExportDirectory_Incoming, ScreenshotUnapprovedFolder);
-
-	// Copy the delta images
-	CopyDirectory(ExportDirectory_Delta, ScreenshotDeltaFolder);
+	CopyDirectory(ExportDirectory, ComparisonResultsFolder);
 
 	return true;
 }
 
-TSharedPtr<FComparisonResults> FScreenShotManager::ImportScreensots(FString ImportPath)
+bool FScreenShotManager::OpenComparisonReports(FString ImportPath, TArray<FComparisonReport>& OutReports)
 {
-	FString ManifestPath = ImportPath / TEXT("ComparisonResults.json");
+	OutReports.Reset();
 
-	FString JsonString;
+	FPaths::NormalizeDirectoryName(ImportPath);
+	ImportPath += TEXT("/");
 
-	if ( FFileHelper::LoadFileToString(JsonString, *ManifestPath) )
-	{
-		TSharedRef< TJsonReader<> > JsonReader = TJsonReaderFactory<>::Create(JsonString);
+	TArray<FString> ComparisonReportPaths;
+	IFileManager::Get().FindFilesRecursive(ComparisonReportPaths, *ImportPath, TEXT("Report.json"), /*Files=*/true, /*Directories=*/false, /*bClearFileNames=*/ false);
 
-		TSharedPtr<FJsonObject> JsonComparisons;
-		if ( !FJsonSerializer::Deserialize(JsonReader, JsonComparisons) )
+	for ( const FString& ReportPath : ComparisonReportPaths )
+	{
+		FString JsonString;
+		if ( FFileHelper::LoadFileToString(JsonString, *ReportPath) )
 		{
-			return nullptr;
-		}
+			TSharedRef< TJsonReader<> > JsonReader = TJsonReaderFactory<>::Create(JsonString);
 
-		FComparisonResults* ImportedResults = new FComparisonResults();
+			TSharedPtr<FJsonObject> JsonComparisonReport;
+			if ( !FJsonSerializer::Deserialize(JsonReader, JsonComparisonReport) )
+			{
+				return false;
+			}
 
-		if ( FJsonObjectConverter::JsonObjectToUStruct(JsonComparisons.ToSharedRef(), ImportedResults, 0, 0) )
-		{
-			return MakeShareable(ImportedResults);
-		}
+			FImageComparisonResult ComparisonResult;
+			if ( FJsonObjectConverter::JsonObjectToUStruct(JsonComparisonReport.ToSharedRef(), &ComparisonResult, 0, 0) )
+			{
+				FComparisonReport Report(ImportPath, ReportPath);
+				Report.Comparison = ComparisonResult;
 
-		delete ImportedResults;
+				OutReports.Add(Report);
+			}
+		}
 	}
 
-	return nullptr;
+	return true;
 }
 
 FString FScreenShotManager::GetDefaultExportDirectory() const
@@ -400,20 +283,6 @@ FString FScreenShotManager::GetDefaultExportDirectory() const
 	return FPaths::GameSavedDir() / TEXT("Exported");
 }
 
-FString FScreenShotManager::GetComparisonResultsJsonFilename() const
-{
-	return ScreenshotResultsFolder / TEXT("ComparisonResults.json");
-}
-
-void FScreenShotManager::SaveComparisonResults(const FComparisonResults& Results) const
-{
-	FString Json;
-	if ( FJsonObjectConverter::UStructToJsonObjectString(Results, Json) )
-	{
-		FFileHelper::SaveStringToFile(Json, *GetComparisonResultsJsonFilename(), FFileHelper::EEncodingOptions::ForceUTF8WithoutBOM);
-	}
-}
-
 void FScreenShotManager::CopyDirectory(const FString& DestDir, const FString& SrcDir)
 {
 	TArray<FString> FilesToCopy;
diff --git a/Engine/Source/Developer/ScreenShotComparisonTools/Private/ScreenShotManager.h b/Engine/Source/Developer/ScreenShotComparisonTools/Private/ScreenShotManager.h
--- a/Engine/Source/Developer/ScreenShotComparisonTools/Private/ScreenShotManager.h
+++ b/Engine/Source/Developer/ScreenShotComparisonTools/Private/ScreenShotManager.h
@@ -29,34 +29,15 @@ public:
 	 */
 	FScreenShotManager();
 
-public:
-
-	/**
-	* Create some dummy data to test the UI
-	*/
-	void CreateData();
-
 public:
 
 	//~ Begin IScreenShotManager Interface
 
-	virtual void GenerateLists() override;
-
-	virtual TArray< TSharedPtr<FString> >& GetCachedPlatfomList() override;
-
-	virtual TArray< TSharedPtr<FScreenShotDataItem> >& GetScreenshotResults() override;
-
-	virtual void RegisterScreenShotUpdate(const FOnScreenFilterChanged& InDelegate ) override;
-
-	virtual void SetFilter( TSharedPtr< ScreenShotFilterCollection > InFilter ) override;
-
-	virtual TFuture<void> CompareScreensotsAsync() override;
-
 	virtual TFuture<FImageComparisonResult> CompareScreensotAsync(FString RelativeImagePath) override;
 
 	virtual TFuture<void> ExportScreensotsAsync(FString ExportPath = TEXT("")) override;
 
-	virtual TSharedPtr<FComparisonResults> ImportScreensots(FString ImportPath) override;
+	virtual bool OpenComparisonReports(FString ImportPath, TArray<FComparisonReport>& OutReports) override;
 
 	virtual FString GetLocalUnapprovedFolder() const override;
 
@@ -66,28 +47,15 @@ public:
 
 private:
 	FString GetDefaultExportDirectory() const;
-	FString GetComparisonResultsJsonFilename() const;
 
-	void CompareScreensots();
 	FImageComparisonResult CompareScreensot(FString Existing);
-	void SaveComparisonResults(const FComparisonResults& Results) const;
 	bool ExportComparisonResults(FString RootExportFolder);
 	void CopyDirectory(const FString& DestDir, const FString& SrcDir);
 
-	TSharedRef<FScreenshotComparisons> GenerateFileList() const;
-
 private:
 	FString ScreenshotApprovedFolder;
 	FString ScreenshotUnapprovedFolder;
 	FString ScreenshotDeltaFolder;
 	FString ScreenshotResultsFolder;
-
-	// Holds the list of active platforms
-	TArray<TSharedPtr<FString> > CachedPlatformList;
-
-	// Holds the array of created screen shot data items
-	TArray< TSharedPtr<FScreenShotDataItem> > ScreenShotDataArray;
-
-	// Holds a delegate to be invoked when the screen shot filter has changed.
-	FOnScreenFilterChanged ScreenFilterChangedDelegate;
+	FString ComparisonResultsFolder;
 };
diff --git a/Engine/Source/Developer/ScreenShotComparisonTools/Public/ImageComparer.h b/Engine/Source/Developer/ScreenShotComparisonTools/Public/ImageComparer.h
--- a/Engine/Source/Developer/ScreenShotComparisonTools/Public/ImageComparer.h
+++ b/Engine/Source/Developer/ScreenShotComparisonTools/Public/ImageComparer.h
@@ -213,6 +213,15 @@ public:
 	UPROPERTY()
 	FString ComparisonFile;
 
+	UPROPERTY()
+	FString ReportApprovedFile;
+
+	UPROPERTY()
+	FString ReportIncomingFile;
+
+	UPROPERTY()
+	FString ReportComparisonFile;
+
 	UPROPERTY()
 	double MaxLocalDifference;
 
@@ -260,46 +269,22 @@ public:
 	}
 };
 
-/**
- *
- */
-USTRUCT()
-struct FComparisonResults
+struct SCREENSHOTCOMPARISONTOOLS_API FComparisonReport
 {
-	GENERATED_USTRUCT_BODY()
-
 public:
 
-	FComparisonResults()
-		: ApprovedPath(TEXT("Approved"))
-		, IncomingPath(TEXT("Incoming"))
-		, DeltaPath(TEXT("Delta"))
-	{
-	}
-
-	UPROPERTY()
-	FString ApprovedPath;
-
-	UPROPERTY()
-	FString IncomingPath;
-
-	UPROPERTY()
-	FString DeltaPath;
-
-	UPROPERTY()
-	TArray<FImageComparisonResult> Comparisons;
+	FComparisonReport(const FString& InReportRootDirectory, const FString& InReportFile);
 
-	UPROPERTY()
-	TArray<FString> Added;
-
-	UPROPERTY()
-	TArray<FString> Missing;
+	FString ReportRootDirectory;
+	FString ReportFile;
+	FString ReportFolder;
+	FImageComparisonResult Comparison;
 };
 
 /**
  * 
  */
-class FImageComparer
+class SCREENSHOTCOMPARISONTOOLS_API FImageComparer
 {
 public:
 	FString ImageRootA;
diff --git a/Engine/Source/Developer/ScreenShotComparisonTools/Public/Interfaces/IScreenShotManager.h b/Engine/Source/Developer/ScreenShotComparisonTools/Public/Interfaces/IScreenShotManager.h
--- a/Engine/Source/Developer/ScreenShotComparisonTools/Public/Interfaces/IScreenShotManager.h
+++ b/Engine/Source/Developer/ScreenShotComparisonTools/Public/Interfaces/IScreenShotManager.h
@@ -9,10 +9,8 @@
 #include "CoreMinimal.h"
 #include "Async/Future.h"
 #include "ImageComparer.h"
-#include "Interfaces/IScreenShotData.h"
 
 class IScreenShotManager;
-struct FScreenShotDataItem;
 
 /**
  * Type definition for shared pointers to instances of IScreenShotManager.
@@ -24,16 +22,6 @@ typedef TSharedPtr<class IScreenShotManager> IScreenShotManagerPtr;
  */
 typedef TSharedRef<class IScreenShotManager> IScreenShotManagerRef;
 
-struct FScreenShotDataItem;
-
-
-
-/**
- * Delegate type for session filter changing.
- */
-DECLARE_DELEGATE(FOnScreenFilterChanged)
-
-
 /**
  * Interface for screen manager module.
  */
@@ -42,41 +30,6 @@ class IScreenShotManager
 public:
 	virtual ~IScreenShotManager(){ }
 
-	/**
-	 * Generate the screen shot data
-	 */
-	virtual void GenerateLists( ) = 0;
-
-	/**
-	 */
-	virtual TArray< TSharedPtr<FScreenShotDataItem> >& GetScreenshotResults() = 0;
-
-	/**
-	 * Get the list of active platforms
-	 *
-	 * @return the array of platforms
-	 */
-	virtual TArray< TSharedPtr<FString> >& GetCachedPlatfomList( ) = 0;
-
-	/**
-	 * Register for screen shot updates
-	 *
-	 * @param InDelegate - Delegate register
-	 */
-	virtual void RegisterScreenShotUpdate( const FOnScreenFilterChanged& InDelegate ) = 0;
-
-	/**
-	 * Set the filter
-	 *
-	 * @param InFilter - The screen shot filter
-	 */
-	virtual void SetFilter(TSharedPtr<ScreenShotFilterCollection> InFilter ) = 0;
-
-	/**
-	 * Compares screenshots.
-	 */
-	virtual TFuture<void> CompareScreensotsAsync() = 0;
-
 	/**
 	 * Compares a specific screenshot, the shot path must be relative from the incoming unapproved directory.
 	 */
@@ -90,9 +43,15 @@ public:
 	/**
 	 * Imports screenshot comparison data from a given path.
 	 */
-	virtual TSharedPtr<FComparisonResults> ImportScreensots(FString ImportPath) = 0;
+	virtual bool OpenComparisonReports(FString ImportPath, TArray<FComparisonReport>& OutReports) = 0;
 
+	/**
+	 * 
+	 */
 	virtual FString GetLocalUnapprovedFolder() const = 0;
 
+	/**
+	 * 
+	 */
 	virtual FString GetLocalApprovedFolder() const = 0;
 };
diff --git a/Engine/Source/Developer/StandaloneRenderer/Private/OpenGL/SlateOpenGLRenderingPolicy.cpp b/Engine/Source/Developer/StandaloneRenderer/Private/OpenGL/SlateOpenGLRenderingPolicy.cpp
--- a/Engine/Source/Developer/StandaloneRenderer/Private/OpenGL/SlateOpenGLRenderingPolicy.cpp
+++ b/Engine/Source/Developer/StandaloneRenderer/Private/OpenGL/SlateOpenGLRenderingPolicy.cpp
@@ -165,9 +165,9 @@ void FSlateOpenGLRenderingPolicy::DrawElements( const FMatrix& ViewProjectionMat
 
 		const FSlateShaderResource* Texture = RenderBatch.Texture;
 
-		const ESlateBatchDrawFlag::Type DrawFlags = RenderBatch.DrawFlags;
+		const ESlateBatchDrawFlag DrawFlags = RenderBatch.DrawFlags;
 
-		if( DrawFlags & ESlateBatchDrawFlag::NoBlending )
+		if( EnumHasAllFlags(DrawFlags, ESlateBatchDrawFlag::NoBlending) )
 		{
 			glDisable( GL_BLEND );
 		}
@@ -177,7 +177,7 @@ void FSlateOpenGLRenderingPolicy::DrawElements( const FMatrix& ViewProjectionMat
 		}
 
 #if !PLATFORM_USES_ES2
-		if( DrawFlags & ESlateBatchDrawFlag::Wireframe )
+		if( EnumHasAllFlags(DrawFlags, ESlateBatchDrawFlag::Wireframe) )
 		{
 			glPolygonMode( GL_FRONT_AND_BACK, GL_LINE );
 			glDisable(GL_BLEND);
@@ -202,8 +202,8 @@ void FSlateOpenGLRenderingPolicy::DrawElements( const FMatrix& ViewProjectionMat
 		}
 		else if( Texture )
 		{
-			uint32 RepeatU = DrawFlags & ESlateBatchDrawFlag::TileU ? GL_REPEAT : GL_CLAMP_TO_EDGE;
-			uint32 RepeatV = DrawFlags & ESlateBatchDrawFlag::TileV ? GL_REPEAT : GL_CLAMP_TO_EDGE;
+			uint32 RepeatU = EnumHasAllFlags(DrawFlags, ESlateBatchDrawFlag::TileU) ? GL_REPEAT : GL_CLAMP_TO_EDGE;
+			uint32 RepeatV = EnumHasAllFlags(DrawFlags, ESlateBatchDrawFlag::TileV) ? GL_REPEAT : GL_CLAMP_TO_EDGE;
 
 #if PLATFORM_USES_ES2
 			if(!FMath::IsPowerOfTwo(Texture->GetWidth()) || !FMath::IsPowerOfTwo(Texture->GetHeight()))
diff --git a/Engine/Source/Developer/StandaloneRenderer/Private/OpenGL/SlateOpenGLShaders.cpp b/Engine/Source/Developer/StandaloneRenderer/Private/OpenGL/SlateOpenGLShaders.cpp
--- a/Engine/Source/Developer/StandaloneRenderer/Private/OpenGL/SlateOpenGLShaders.cpp
+++ b/Engine/Source/Developer/StandaloneRenderer/Private/OpenGL/SlateOpenGLShaders.cpp
@@ -312,10 +312,10 @@ void FSlateOpenGLElementProgram::SetVertexShaderParams( const FVector4& ShaderPa
 	CHECK_GL_ERRORS;
 }
 
-void FSlateOpenGLElementProgram::SetDrawEffects( uint32 InDrawEffects )
+void FSlateOpenGLElementProgram::SetDrawEffects(ESlateDrawEffect InDrawEffects )
 {
-	glUniform1i( EffectsDisabledParam, InDrawEffects & ESlateDrawEffect::DisabledEffect );
-	glUniform1i( IgnoreTextureAlphaParam, InDrawEffects & ESlateDrawEffect::IgnoreTextureAlpha );
+	glUniform1i( EffectsDisabledParam, EnumHasAllFlags(InDrawEffects, ESlateDrawEffect::DisabledEffect) ? 1 : 0 );
+	glUniform1i( IgnoreTextureAlphaParam, EnumHasAllFlags(InDrawEffects, ESlateDrawEffect::IgnoreTextureAlpha) ? 1 : 0 );
 	CHECK_GL_ERRORS;
 }
 
diff --git a/Engine/Source/Developer/StandaloneRenderer/Private/OpenGL/SlateOpenGLShaders.h b/Engine/Source/Developer/StandaloneRenderer/Private/OpenGL/SlateOpenGLShaders.h
--- a/Engine/Source/Developer/StandaloneRenderer/Private/OpenGL/SlateOpenGLShaders.h
+++ b/Engine/Source/Developer/StandaloneRenderer/Private/OpenGL/SlateOpenGLShaders.h
@@ -7,6 +7,7 @@
 
 #include "CoreMinimal.h"
 #include "StandaloneRendererPlatformHeaders.h"
+#include "Rendering/RenderingCommon.h"
 
 /**
  * Base class for all OpenGL shaders                   
@@ -66,7 +67,7 @@ public:
 	void SetViewProjectionMatrix( const FMatrix& InVP ); 
 	void SetVertexShaderParams( const FVector4& ShaderParams );
 	void SetTexture( GLuint Texture, uint32 AddressU, uint32 AddressV  );
-	void SetDrawEffects( uint32 InDrawEffects );
+	void SetDrawEffects(ESlateDrawEffect InDrawEffects );
 	void SetShaderType( uint32 InShaderType );
 	void SetMarginUVs( const FVector4& InMarginUVs );
 private:
diff --git a/Engine/Source/Developer/StandaloneRenderer/Private/Windows/D3D/SlateD3DRenderingPolicy.cpp b/Engine/Source/Developer/StandaloneRenderer/Private/Windows/D3D/SlateD3DRenderingPolicy.cpp
--- a/Engine/Source/Developer/StandaloneRenderer/Private/Windows/D3D/SlateD3DRenderingPolicy.cpp
+++ b/Engine/Source/Developer/StandaloneRenderer/Private/Windows/D3D/SlateD3DRenderingPolicy.cpp
@@ -234,15 +234,15 @@ void FSlateD3D11RenderingPolicy::DrawElements( const FMatrix& ViewProjectionMatr
 
 		const FSlateShaderResource* Texture = RenderBatch.Texture;
 		
-		const ESlateBatchDrawFlag::Type DrawFlags = RenderBatch.DrawFlags;
+		const ESlateBatchDrawFlag DrawFlags = RenderBatch.DrawFlags;
 
-		const ESlateDrawEffect::Type DrawEffects = RenderBatch.DrawEffects;
+		const ESlateDrawEffect DrawEffects = RenderBatch.DrawEffects;
 
 		const FShaderParams& ShaderParams = RenderBatch.ShaderParams;
 
 		VertexShader->BindParameters();
 		
-		if( DrawFlags & ESlateBatchDrawFlag::NoBlending )
+		if( EnumHasAllFlags(DrawFlags, ESlateBatchDrawFlag::NoBlending) )
 		{
 			GD3DDeviceContext->OMSetBlendState( NoBlendState, 0, 0xFFFFFFFF );
 		}
@@ -251,7 +251,7 @@ void FSlateD3D11RenderingPolicy::DrawElements( const FMatrix& ViewProjectionMatr
 			GD3DDeviceContext->OMSetBlendState( AlphaBlendState, 0, 0xFFFFFFFF );
 		}
 
-		if( DrawFlags & ESlateBatchDrawFlag::Wireframe )
+		if( EnumHasAllFlags(DrawFlags, ESlateBatchDrawFlag::Wireframe) )
 		{
 			GD3DDeviceContext->RSSetState( WireframeRasterState );
 		}
@@ -283,7 +283,7 @@ void FSlateD3D11RenderingPolicy::DrawElements( const FMatrix& ViewProjectionMatr
 		if( Texture )
 		{
 			TRefCountPtr<ID3D11SamplerState> SamplerState;
-			if( DrawFlags & ESlateBatchDrawFlag::TileU || DrawFlags & ESlateBatchDrawFlag::TileV )
+			if( EnumHasAllFlags(DrawFlags, ESlateBatchDrawFlag::TileU) || EnumHasAllFlags(DrawFlags, ESlateBatchDrawFlag::TileV) )
 			{
 				SamplerState = BilinearSamplerState_Wrap;
 			}
diff --git a/Engine/Source/Developer/StandaloneRenderer/Private/Windows/D3D/SlateD3DShaders.cpp b/Engine/Source/Developer/StandaloneRenderer/Private/Windows/D3D/SlateD3DShaders.cpp
--- a/Engine/Source/Developer/StandaloneRenderer/Private/Windows/D3D/SlateD3DShaders.cpp
+++ b/Engine/Source/Developer/StandaloneRenderer/Private/Windows/D3D/SlateD3DShaders.cpp
@@ -374,9 +374,9 @@ void FSlateDefaultPS::SetShaderType( uint32 InShaderType )
 	PerElementConstants.GetBufferData().ShaderType = InShaderType;
 }
 
-void FSlateDefaultPS::SetDrawEffects( uint32 InDrawEffects )
+void FSlateDefaultPS::SetDrawEffects( ESlateDrawEffect InDrawEffects )
 {
-	PerElementConstants.GetBufferData().DrawEffects = InDrawEffects;
+	PerElementConstants.GetBufferData().DrawEffects = (uint32)InDrawEffects;
 }
 
 void FSlateDefaultPS::SetShaderParams( const FVector4& InShaderParams )
diff --git a/Engine/Source/Developer/StandaloneRenderer/Private/Windows/D3D/SlateD3DShaders.h b/Engine/Source/Developer/StandaloneRenderer/Private/Windows/D3D/SlateD3DShaders.h
--- a/Engine/Source/Developer/StandaloneRenderer/Private/Windows/D3D/SlateD3DShaders.h
+++ b/Engine/Source/Developer/StandaloneRenderer/Private/Windows/D3D/SlateD3DShaders.h
@@ -5,6 +5,7 @@
 #include "CoreMinimal.h"
 #include "StandaloneRendererPlatformHeaders.h"
 #include "SlateD3DConstantBuffer.h"
+#include "Rendering/RenderingCommon.h"
 
 class FSlateD3DShaderParameter
 {
@@ -173,7 +174,7 @@ public:
 	}
 
 	void SetShaderType( uint32 InShaderType );
-	void SetDrawEffects( uint32 InDrawEffects );
+	void SetDrawEffects(ESlateDrawEffect InDrawEffects );
 	void SetShaderParams( const FVector4& InShaderParams );
 	void SetGammaValues(const FVector2D& InGammaValues);
 protected:
diff --git a/Engine/Source/Editor/Cascade/Public/CascadeModule.h b/Engine/Source/Editor/Cascade/Public/CascadeModule.h
--- a/Engine/Source/Editor/Cascade/Public/CascadeModule.h
+++ b/Engine/Source/Editor/Cascade/Public/CascadeModule.h
@@ -1,7 +1,6 @@
 // Copyright 1998-2017 Epic Games, Inc. All Rights Reserved.
 
-#ifndef __CascadeModule_h__
-#define __CascadeModule_h__
+#pragma once
 
 #include "CoreMinimal.h"
 #include "Modules/ModuleInterface.h"
@@ -13,9 +12,6 @@ class UParticleSystem;
 
 extern const FName CascadeAppIdentifier;
 
-class FCascade;
-
-
 /*-----------------------------------------------------------------------------
    ICascadeModule
 -----------------------------------------------------------------------------*/
@@ -36,5 +32,3 @@ public:
 	/** Converts all the modules in the specified particle system to seeded modules */
 	virtual void ConvertModulesToSeeded(UParticleSystem* ParticleSystem) = 0;
 };
-
-#endif // __CascadeModule_h__
diff --git a/Engine/Source/Editor/ContentBrowser/Private/AssetViewWidgets.cpp b/Engine/Source/Editor/ContentBrowser/Private/AssetViewWidgets.cpp
--- a/Engine/Source/Editor/ContentBrowser/Private/AssetViewWidgets.cpp
+++ b/Engine/Source/Editor/ContentBrowser/Private/AssetViewWidgets.cpp
@@ -1632,7 +1632,7 @@ TSharedRef<SWidget> SAssetColumnItem::GenerateWidgetForColumn( const FName& Colu
 	TSharedPtr<SWidget> Content;
 
 	// A little right padding so text from this column does not run directly into text from the next.
-	static const FMargin ColumnItemPadding( 0, 0, 6, 0 );
+	static const FMargin ColumnItemPadding( 5, 0, 5, 0 );
 
 	if ( ColumnName == "Name" )
 	{
@@ -1650,7 +1650,7 @@ TSharedRef<SWidget> SAssetColumnItem::GenerateWidgetForColumn( const FName& Colu
 		}
 		else
 		{
-			IconBrush = FEditorStyle::GetBrush("ContentBrowser.ColumnViewAssetIcon");;
+			IconBrush = FEditorStyle::GetBrush("ContentBrowser.ColumnViewAssetIcon");
 		}
 
 		// Make icon overlays (eg, SCC and dirty status) a reasonable size in relation to the icon size (note: it is assumed this icon is square)
@@ -1724,6 +1724,8 @@ TSharedRef<SWidget> SAssetColumnItem::GenerateWidgetForColumn( const FName& Colu
 		return SNew(SBorder)
 			.BorderImage(this, &SAssetViewItem::GetBorderImage)
 			.Padding(0)
+			.VAlign(VAlign_Center)
+			.HAlign(HAlign_Left)
 			[
 				SNew( SAssetColumnItemNameBox, SharedThis(this) )
 				.Padding( ColumnItemPadding )
@@ -1755,6 +1757,8 @@ TSharedRef<SWidget> SAssetColumnItem::GenerateWidgetForColumn( const FName& Colu
 
 	return SNew(SBox)
 		.Padding( ColumnItemPadding )
+		.VAlign(VAlign_Center)
+		.HAlign(HAlign_Left)
 		[
 			Content.ToSharedRef()
 		];
diff --git a/Engine/Source/Editor/ContentBrowser/Private/SAssetView.cpp b/Engine/Source/Editor/ContentBrowser/Private/SAssetView.cpp
--- a/Engine/Source/Editor/ContentBrowser/Private/SAssetView.cpp
+++ b/Engine/Source/Editor/ContentBrowser/Private/SAssetView.cpp
@@ -45,6 +45,7 @@
 #include "NativeClassHierarchy.h"
 #include "Framework/Notifications/NotificationManager.h"
 #include "Widgets/Notifications/SNotificationList.h"
+#include "SSplitter.h"
 
 #define LOCTEXT_NAMESPACE "ContentBrowser"
 
@@ -1544,6 +1545,7 @@ TSharedRef<SAssetColumnView> SAssetView::CreateColumnView()
 		.HeaderRow
 		(
 			SNew(SHeaderRow)
+			.ResizeMode(ESplitterResizeMode::FixedSize)
 			+ SHeaderRow::Column(SortManager.NameColumnId)
 			.FillWidth(300)
 			.SortMode( TAttribute< EColumnSortMode::Type >::Create( TAttribute< EColumnSortMode::Type >::FGetter::CreateSP( this, &SAssetView::GetColumnSortMode, SortManager.NameColumnId ) ) )
@@ -1557,6 +1559,9 @@ TSharedRef<SAssetColumnView> SAssetView::CreateColumnView()
 			]
 		);
 
+	NewColumnView->GetHeaderRow()->SetOnGetMaxRowSizeForColumn(FOnGetMaxRowSizeForColumn::CreateRaw(NewColumnView.Get(), &SAssetColumnView::GetMaxRowSizeForColumn));
+
+
 	NumVisibleColumns = HiddenColumnNames.Contains(SortManager.NameColumnId.ToString()) ? 0 : 1;
 
 	if(bShowTypeInColumnView)
@@ -2204,7 +2209,6 @@ void SAssetView::SetMajorityAssetType(FName NewMajorityAssetType)
 									.OnSort(FOnSortModeChanged::CreateSP(this, &SAssetView::OnSortColumnHeader))
 									.DefaultLabel(DisplayName)
 									.DefaultTooltip(TooltipText)
-									.HAlignCell((TagIt->Type == UObject::FAssetRegistryTag::TT_Numerical) ? HAlign_Right : HAlign_Left)
 									.FillWidth(180)
 									.ShouldGenerateWidget(TAttribute<bool>::Create(TAttribute<bool>::FGetter::CreateSP(this, &SAssetView::ShouldColumnGenerateWidget, TagName.ToString())))
 									.MenuContent()
diff --git a/Engine/Source/Editor/DataTableEditor/Private/SRowEditor.cpp b/Engine/Source/Editor/DataTableEditor/Private/SRowEditor.cpp
--- a/Engine/Source/Editor/DataTableEditor/Private/SRowEditor.cpp
+++ b/Engine/Source/Editor/DataTableEditor/Private/SRowEditor.cpp
@@ -351,6 +351,31 @@ void SRowEditor::OnRowRenamed(const FText& Text, ETextCommit::Type CommitType)
 	}
 }
 
+FReply SRowEditor::OnResetToDefaultClicked()
+{
+	if (DataTable.IsValid() && SelectedName.IsValid())
+	{
+		FDataTableEditorUtils::ResetToDefault(DataTable.Get(), *SelectedName.Get());
+	}
+
+	return FReply::Handled();
+}
+
+EVisibility SRowEditor::GetResetToDefaultVisibility() const
+{
+	EVisibility VisibleState = EVisibility::Collapsed;
+
+	if (DataTable.IsValid() && SelectedName.IsValid())
+	{
+		if (FDataTableEditorUtils::DiffersFromDefault(DataTable.Get(), *SelectedName.Get()))
+		{
+			VisibleState = EVisibility::Visible;
+		}
+	}
+
+	return VisibleState;
+}
+
 void SRowEditor::Construct(const FArguments& InArgs, UDataTable* Changed)
 {
 	DataTable = Changed;
@@ -431,6 +456,26 @@ void SRowEditor::Construct(const FArguments& InArgs, UDataTable* Changed)
 					]
 				]
 			]
+			+ SHorizontalBox::Slot()
+			.AutoWidth()
+			.Padding(2)
+			[
+				SNew(SButton)
+				.OnClicked(this, &SRowEditor::OnResetToDefaultClicked)
+				.Visibility(this, &SRowEditor::GetResetToDefaultVisibility)
+				.ContentPadding(FMargin(5.f, 0.f))
+				.ToolTipText(LOCTEXT("ResetToDefaultToolTip", "Reset to Default"))
+				.ButtonStyle(FEditorStyle::Get(), "NoBorder")
+				.ForegroundColor(FSlateColor::UseForeground())
+				.HAlign(HAlign_Center)
+				.VAlign(VAlign_Center)
+				.Content()
+				[
+					SNew(SImage)
+					.Image(FEditorStyle::GetBrush("PropertyWindow.DiffersFromDefault"))
+				]
+			]
+
 			+ SHorizontalBox::Slot()
 			[
 				SNew(SSpacer)
diff --git a/Engine/Source/Editor/DataTableEditor/Private/SRowEditor.h b/Engine/Source/Editor/DataTableEditor/Private/SRowEditor.h
--- a/Engine/Source/Editor/DataTableEditor/Private/SRowEditor.h
+++ b/Engine/Source/Editor/DataTableEditor/Private/SRowEditor.h
@@ -71,6 +71,8 @@ private:
 	FReply OnMoveRowClicked(FDataTableEditorUtils::ERowMoveDirection MoveDirection);
 	FReply OnMoveToExtentClicked(FDataTableEditorUtils::ERowMoveDirection MoveDirection);
 	void OnRowRenamed(const FText& Text, ETextCommit::Type CommitType);
+	FReply OnResetToDefaultClicked();
+	EVisibility GetResetToDefaultVisibility()const ;
 
 public:
 
diff --git a/Engine/Source/Editor/DetailCustomizations/Private/ActorDetails.cpp b/Engine/Source/Editor/DetailCustomizations/Private/ActorDetails.cpp
--- a/Engine/Source/Editor/DetailCustomizations/Private/ActorDetails.cpp
+++ b/Engine/Source/Editor/DetailCustomizations/Private/ActorDetails.cpp
@@ -50,6 +50,7 @@
 #include "ActorDetailsDelegates.h"
 #include "EditorCategoryUtils.h"
 #include "Widgets/Input/SHyperlink.h"
+#include "ObjectEditorUtils.h"
 
 #define LOCTEXT_NAMESPACE "ActorDetails"
 
@@ -458,20 +459,23 @@ void FActorDetails::AddTransformCategory( IDetailLayoutBuilder& DetailBuilder )
 void FActorDetails::AddExperimentalWarningCategory( IDetailLayoutBuilder& DetailBuilder )
 {
 	const FSelectedActorInfo& SelectedActorInfo = DetailBuilder.GetDetailsView().GetSelectedActorInfo();
+	bool bBaseClassIsExperimental = false;
+	bool bBaseClassIsEarlyAccess = false;
 
-	if (SelectedActorInfo.bHaveExperimentalClass || SelectedActorInfo.bHaveEarlyAccessClass)
+	FObjectEditorUtils::GetClassDevelopmentStatus(DetailBuilder.GetBaseClass(), bBaseClassIsExperimental, bBaseClassIsEarlyAccess);
+
+	if (bBaseClassIsExperimental || bBaseClassIsEarlyAccess)
 	{
-		const bool bExperimental = SelectedActorInfo.bHaveExperimentalClass;
 
 		const FName CategoryName(TEXT("Warning"));
 		const FText CategoryDisplayName = LOCTEXT("WarningCategoryDisplayName", "Warning");
 		FString ClassUsed = DetailBuilder.GetTopLevelProperty().ToString();
-		const FText WarningText = bExperimental ? FText::Format( LOCTEXT("ExperimentalClassWarning", "Uses experimental class: {0}") , FText::FromString(ClassUsed) )
+		const FText WarningText = bBaseClassIsExperimental ? FText::Format( LOCTEXT("ExperimentalClassWarning", "Uses experimental class: {0}") , FText::FromString(ClassUsed) )
 			: FText::Format( LOCTEXT("EarlyAccessClassWarning", "Uses early access class {0}"), FText::FromString(*ClassUsed) );
 		const FText SearchString = WarningText;
-		const FText Tooltip = bExperimental ? LOCTEXT("ExperimentalClassTooltip", "Here be dragons!  Uses one or more unsupported 'experimental' classes") : LOCTEXT("EarlyAccessClassTooltip", "Uses one or more 'early access' classes");
-		const FString ExcerptName = bExperimental ? TEXT("ActorUsesExperimentalClass") : TEXT("ActorUsesEarlyAccessClass");
-		const FSlateBrush* WarningIcon = FEditorStyle::GetBrush(bExperimental ? "PropertyEditor.ExperimentalClass" : "PropertyEditor.EarlyAccessClass");
+		const FText Tooltip = bBaseClassIsExperimental ? LOCTEXT("ExperimentalClassTooltip", "Here be dragons!  Uses one or more unsupported 'experimental' classes") : LOCTEXT("EarlyAccessClassTooltip", "Uses one or more 'early access' classes");
+		const FString ExcerptName = bBaseClassIsExperimental ? TEXT("ActorUsesExperimentalClass") : TEXT("ActorUsesEarlyAccessClass");
+		const FSlateBrush* WarningIcon = FEditorStyle::GetBrush(bBaseClassIsExperimental ? "PropertyEditor.ExperimentalClass" : "PropertyEditor.EarlyAccessClass");
 
 		IDetailCategoryBuilder& WarningCategory = DetailBuilder.EditCategory(CategoryName, CategoryDisplayName, ECategoryPriority::Transform);
 
diff --git a/Engine/Source/Editor/DetailCustomizations/Private/AutoReimportDirectoryCustomization.cpp b/Engine/Source/Editor/DetailCustomizations/Private/AutoReimportDirectoryCustomization.cpp
--- a/Engine/Source/Editor/DetailCustomizations/Private/AutoReimportDirectoryCustomization.cpp
+++ b/Engine/Source/Editor/DetailCustomizations/Private/AutoReimportDirectoryCustomization.cpp
@@ -6,16 +6,20 @@
 #include "Engine/GameViewportClient.h"
 #include "Misc/PackageName.h"
 #include "Widgets/SBoxPanel.h"
+#include "Widgets/Layout/SBox.h"
 #include "Styling/SlateTypes.h"
 #include "Widgets/Text/STextBlock.h"
 #include "PropertyHandle.h"
 #include "Widgets/Input/SEditableTextBox.h"
 #include "Widgets/Input/SButton.h"
 #include "Widgets/Input/SCheckBox.h"
+#include "Widgets/Input/SComboButton.h"
 #include "Settings/EditorLoadingSavingSettings.h"
 #include "IDetailChildrenBuilder.h"
 #include "DetailWidgetRow.h"
 #include "DetailLayoutBuilder.h"
+#include "IContentBrowserSingleton.h"
+#include "ContentBrowserModule.h"
 
 #include "DesktopPlatformModule.h"
 
@@ -137,6 +141,7 @@ void FAutoReimportDirectoryCustomization::CustomizeHeader(TSharedRef<IPropertyHa
 		.AutoWidth()
 		[
 			SNew(SButton)
+			.ContentPadding(FMargin(4.0f, 2.0f))
 			.OnClicked(this, &FAutoReimportDirectoryCustomization::BrowseForFolder)
 			.ToolTipText(LOCTEXT("BrowseForDirectory", "Browse for a directory"))
 			.Text(LOCTEXT("Browse", "Browse"))
@@ -167,15 +172,26 @@ void FAutoReimportDirectoryCustomization::CustomizeChildren(TSharedRef<IProperty
 			SNew(SHorizontalBox)
 
 			+SHorizontalBox::Slot()
+			.Padding(FMargin(0.0f, 0.0f, 4.0f, 0.0f))
+			.AutoWidth()
+			.VAlign(VAlign_Center)
 			[
-				SNew(SOverlay)
+				SNew(STextBlock)
+				.Text(this, &FAutoReimportDirectoryCustomization::GetMountPointText)
+			]
 
-				+ SOverlay::Slot()
+			+SHorizontalBox::Slot()
+			.AutoWidth()
+			[
+				SAssignNew(PathPickerButton, SComboButton)
+				.HasDownArrow(false)
+				.ToolTipText(LOCTEXT("BrowseForMountPoint", "Choose a path"))
+				.OnGetMenuContent(this, &FAutoReimportDirectoryCustomization::GetPathPickerContent)
+				.ContentPadding(FMargin(4.0f, 2.0f))
+				.ButtonContent()
 				[
-					SNew(SEditableTextBox)
-					.Text(this, &FAutoReimportDirectoryCustomization::GetMountPointText)
-					.OnTextChanged(this, &FAutoReimportDirectoryCustomization::OnMountPointChanged)
-					.OnTextCommitted(this, &FAutoReimportDirectoryCustomization::OnMountPointCommitted)
+					SNew(STextBlock)
+					.Text(LOCTEXT("Browse", "Browse"))
 				]
 			]
 		];
@@ -184,6 +200,43 @@ void FAutoReimportDirectoryCustomization::CustomizeChildren(TSharedRef<IProperty
 	InStructBuilder.AddChildProperty(WildcardsProperty.ToSharedRef());
 }
 
+TSharedRef<SWidget> FAutoReimportDirectoryCustomization::GetPathPickerContent()
+{
+	FPathPickerConfig PathPickerConfig;
+
+	MountPointProperty->GetValueAsFormattedString(PathPickerConfig.DefaultPath);
+	PathPickerConfig.DefaultPath.ReplaceInline(TEXT("\\"), TEXT("/"));
+
+	PathPickerConfig.OnPathSelected = FOnPathSelected::CreateSP(this, &FAutoReimportDirectoryCustomization::PathPickerPathSelected);
+	PathPickerConfig.bAllowClassesFolder = false;
+	PathPickerConfig.bAddDefaultPath = false;
+	PathPickerConfig.bAllowContextMenu = false;
+	PathPickerConfig.bFocusSearchBoxWhenOpened = false;
+
+	FContentBrowserModule& ContentBrowserModule = FModuleManager::LoadModuleChecked<FContentBrowserModule>("ContentBrowser");
+
+	return SNew(SBox)
+		.WidthOverride(300)
+		.HeightOverride(500)
+		.Padding(4)
+		[
+			SNew(SVerticalBox)
+
+			+ SVerticalBox::Slot()
+			.FillHeight(1.f)
+			[
+				ContentBrowserModule.Get().CreatePathPicker(PathPickerConfig)
+			]
+		];
+}
+
+void FAutoReimportDirectoryCustomization::PathPickerPathSelected(const FString& FolderPath)
+{
+	PathPickerButton->SetIsOpen(false);
+
+	MountPointProperty->SetValue(FolderPath);
+}
+
 EVisibility FAutoReimportDirectoryCustomization::GetMountPathVisibility() const
 {
 	return MountPathVisibility;
@@ -213,16 +266,6 @@ FText FAutoReimportDirectoryCustomization::GetMountPointText() const
 	return Ret;
 }
 
-void FAutoReimportDirectoryCustomization::OnMountPointCommitted(const FText& InValue, ETextCommit::Type CommitType)
-{
-	MountPointProperty->SetValue(InValue.ToString());
-}
-
-void FAutoReimportDirectoryCustomization::OnMountPointChanged(const FText& InValue)
-{
-	MountPointProperty->SetValue(InValue.ToString());
-}
-
 FReply FAutoReimportDirectoryCustomization::BrowseForFolder()
 {
 	FString InitialDir;
diff --git a/Engine/Source/Editor/DetailCustomizations/Private/AutoReimportDirectoryCustomization.h b/Engine/Source/Editor/DetailCustomizations/Private/AutoReimportDirectoryCustomization.h
--- a/Engine/Source/Editor/DetailCustomizations/Private/AutoReimportDirectoryCustomization.h
+++ b/Engine/Source/Editor/DetailCustomizations/Private/AutoReimportDirectoryCustomization.h
@@ -70,16 +70,18 @@ private:
 	void OnDirectoryChanged(const FText& InValue);
 
 	FText GetMountPointText() const;
-	void OnMountPointCommitted(const FText& InValue, ETextCommit::Type CommitType);
-	void OnMountPointChanged(const FText& InValue);
 
 	void SetSourcePath(FString InSourceDir);
 	void UpdateMountPath();
 	FReply BrowseForFolder();
 
+	TSharedRef<class SWidget> GetPathPickerContent();
+	void PathPickerPathSelected(const FString& FolderPath);
+
 	EVisibility MountPathVisibility;
 	TSharedPtr<IPropertyHandle> PropertyHandle;
 	TSharedPtr<IPropertyHandle> SourceDirProperty;
 	TSharedPtr<IPropertyHandle> MountPointProperty;
 	TSharedPtr<IPropertyHandle> WildcardsProperty;
+	TSharedPtr<class SComboButton> PathPickerButton;
 };
diff --git a/Engine/Source/Editor/DetailCustomizations/Private/ColorGradingVectorCustomization.cpp b/Engine/Source/Editor/DetailCustomizations/Private/ColorGradingVectorCustomization.cpp
--- a/Engine/Source/Editor/DetailCustomizations/Private/ColorGradingVectorCustomization.cpp
+++ b/Engine/Source/Editor/DetailCustomizations/Private/ColorGradingVectorCustomization.cpp
@@ -168,11 +168,21 @@ void FColorGradingCustomBuilder::GenerateHeaderRowContent(FDetailWidgetRow& Node
 			SNew(SBox)
 		];
 	int32 ColorIndex = 0;
+	FText LabelText[4];
+	LabelText[0] = NSLOCTEXT("ColorGradingVectorCustomizationNS", "RedChannelSmallName", "R");
+	LabelText[1] = NSLOCTEXT("ColorGradingVectorCustomizationNS", "GreenChannelSmallName", "G");
+	LabelText[2] = NSLOCTEXT("ColorGradingVectorCustomizationNS", "BlueChannelSmallName", "B");
+	LabelText[3] = NSLOCTEXT("ColorGradingVectorCustomizationNS", "LuminanceChannelSmallName", "Y");
 	for (TWeakPtr<IPropertyHandle> ChildProperty : SortedChildArray)
 	{
+		FText LabelChannelText = ChildProperty.Pin()->GetPropertyDisplayName();
+		if(ColorIndex >= 0 && ColorIndex < 4)
+		{
+			LabelChannelText = LabelText[ColorIndex];
+		}
 		TSharedPtr<SWidget> LabelWidget = SNew(STextBlock)
 			.Font(IDetailLayoutBuilder::GetDetailFont())
-			.Text(ChildProperty.Pin()->GetPropertyDisplayName());
+			.Text(LabelChannelText);
 
 
 		TSharedPtr<SNumericEntryBox<float>> NumericEntry;
diff --git a/Engine/Source/Editor/DetailCustomizations/Private/MathStructCustomizations.cpp b/Engine/Source/Editor/DetailCustomizations/Private/MathStructCustomizations.cpp
--- a/Engine/Source/Editor/DetailCustomizations/Private/MathStructCustomizations.cpp
+++ b/Engine/Source/Editor/DetailCustomizations/Private/MathStructCustomizations.cpp
@@ -156,7 +156,7 @@ void ExtractNumericMetadata(TSharedRef<IPropertyHandle>& PropertyHandle, TOption
 	const FString& ShiftMouseMovePixelPerDeltaString = Property->GetMetaData(TEXT("ShiftMouseMovePixelPerDelta"));
 	const FString& ClampMinString = Property->GetMetaData(TEXT("ClampMin"));
 	const FString& ClampMaxString = Property->GetMetaData(TEXT("ClampMax"));
-	const FString& ColorGradingModeString = Property->GetMetaData(TEXT("ColorGradingMode"));
+	const FString& DisplayColorChannelString = Property->GetMetaData(TEXT("DisplayColorChannel"));
 
 	// If no UIMin/Max was specified then use the clamp string
 	const FString& UIMinString = MetaUIMinString.Len() ? MetaUIMinString : ClampMinString;
@@ -224,7 +224,7 @@ void ExtractNumericMetadata(TSharedRef<IPropertyHandle>& PropertyHandle, TOption
 		//UE_LOG(LogPropertyNode, Warning, TEXT("UI Min (%s) >= UI Max (%s) for Ranged Numeric"), *UIMinString, *UIMaxString);
 	}
 
-	DisplayChannelColor = ColorGradingModeString.Len() > 0 ? true : false;
+	DisplayChannelColor = DisplayColorChannelString.Len() > 0 && DisplayColorChannelString.ToBool();
 }
 
 
@@ -239,10 +239,6 @@ TSharedRef<SWidget> FMathStructCustomization::MakeNumericWidget(
 	bool DisplayChannelColor = false;
 	ExtractNumericMetadata(StructurePropertyHandle, MinValue, MaxValue, SliderMinValue, SliderMaxValue, SliderExponent, Delta, ShiftMouseMovePixelPerDelta, DisplayChannelColor);
 
-	TSharedPtr<SWidget> LabelWidget = SNew(STextBlock)
-		.Font(IDetailLayoutBuilder::GetDetailFont())
-		.Text(PropertyHandle->GetPropertyDisplayName());
-
 	int32 ColorIndex = INDEX_NONE;
 	if (DisplayChannelColor && SortedChildHandles.Num() > 2)
 	{
@@ -257,6 +253,30 @@ TSharedRef<SWidget> FMathStructCustomization::MakeNumericWidget(
 		}
 	}
 
+	//In case we have to display the struct has a color we rename the label RGBY instead of XYZW
+	FText LabelText = PropertyHandle->GetPropertyDisplayName();
+	if (ColorIndex != INDEX_NONE)
+	{
+		switch (ColorIndex)
+		{
+		case 0:
+			LabelText = NSLOCTEXT("MathStrucCustomizationNS", "RedChannelSmallName", "R");
+			break;
+		case 1:
+			LabelText = NSLOCTEXT("MathStrucCustomizationNS", "GreenChannelSmallName", "G");
+			break;
+		case 2:
+			LabelText = NSLOCTEXT("MathStrucCustomizationNS", "BlueChannelSmallName", "B");
+			break;
+		case 3:
+			LabelText = NSLOCTEXT("MathStrucCustomizationNS", "LuminanceChannelSmallName", "Y");
+			break;
+		}
+	}
+	TSharedPtr<SWidget> LabelWidget = SNew(STextBlock)
+		.Font(IDetailLayoutBuilder::GetDetailFont())
+		.Text(LabelText);
+
 	TWeakPtr<IPropertyHandle> WeakHandlePtr = PropertyHandle;
 
 	const FEditableTextBoxStyle& DarkEditorStyle = FCoreStyle::Get().GetWidgetStyle<FEditableTextBoxStyle>("DarkEditableTextBox");
diff --git a/Engine/Source/Editor/DetailCustomizations/Private/Vector4StructCustomization.cpp b/Engine/Source/Editor/DetailCustomizations/Private/Vector4StructCustomization.cpp
--- a/Engine/Source/Editor/DetailCustomizations/Private/Vector4StructCustomization.cpp
+++ b/Engine/Source/Editor/DetailCustomizations/Private/Vector4StructCustomization.cpp
@@ -7,6 +7,7 @@
 #include "DetailLayoutBuilder.h"
 #include "DetailWidgetRow.h"
 #include "UObject/UnrealType.h"
+#include "Widgets/Layout/SBox.h"
 
 TSharedRef<IPropertyTypeCustomization> FVector4StructCustomization::MakeInstance()
 {
@@ -60,5 +61,68 @@ void FVector4StructCustomization::MakeHeaderRow(TSharedRef<class IPropertyHandle
 {
 	//Use the base class header row
 	FMathStructCustomization::MakeHeaderRow(StructPropertyHandle, Row);
+
+	//Modify the name content if the vector is representing some color grading.
+	//We are not showing the color preview in the case of an additive color grading vector
+	//like the offset since additive value as no color meaning
+	UProperty* Property = StructPropertyHandle->GetProperty();
+	if (Property)
+	{
+		const FString& ColorGradingModeString = Property->GetMetaData(TEXT("ColorGradingMode"));
+		bool IsAdditiveMode = (ColorGradingModeString.Compare(TEXT("offset")) == 0);
+		if (ColorGradingModeString.Len() > 0 && !IsAdditiveMode)
+		{
+			//Get a weak pointer of the Vector4 property handle to retrieve the value for the color box preview
+			TWeakPtr<IPropertyHandle> StructWeakHandlePtr = StructPropertyHandle;
+
+			//Get the original name content
+			TSharedRef<SWidget> OriginalNameContent = Row.NameWidget.Widget;
+
+			//Put back the original content on the left side of the NameContent box
+			TSharedPtr<SHorizontalBox> NameContentHorizontalBox;
+			SAssignNew(NameContentHorizontalBox, SHorizontalBox)
+				+ SHorizontalBox::Slot()
+				.AutoWidth()
+				.HAlign(HAlign_Left)
+				[
+					OriginalNameContent
+				];
+
+			//Add the color box preview on the right side of the name content box
+			NameContentHorizontalBox->AddSlot()
+				.HAlign(HAlign_Right)
+				.FillWidth(1.0)
+				[
+					SNew(SBox)
+					.HAlign(HAlign_Right)
+				.Padding(FMargin(0.0f, 0.0f, 5.0f, 0.0f))
+				[
+					SNew(SColorBlock)
+					.Color(this, &FVector4StructCustomization::OnGetHeaderColorBlock, StructWeakHandlePtr)
+					.ShowBackgroundForAlpha(false)
+					.IgnoreAlpha(true)
+					.Size(FVector2D(35.0f, 12.0f))
+					]
+				];
+
+			//Override the old NameContent with the new NameContent box
+			Row.NameContent()
+				[
+					NameContentHorizontalBox.ToSharedRef()
+				];
+		}
+	}
 }
 
+FLinearColor FVector4StructCustomization::OnGetHeaderColorBlock(TWeakPtr<IPropertyHandle> WeakHandlePtr) const
+{
+	FLinearColor ColorValue(0.0f, 0.0f, 0.0f);
+	FVector4 VectorValue;
+	if (WeakHandlePtr.Pin()->GetValue(VectorValue) == FPropertyAccess::Success)
+	{
+		ColorValue.R = VectorValue.X * VectorValue.W;
+		ColorValue.G = VectorValue.Y * VectorValue.W;
+		ColorValue.B = VectorValue.Z * VectorValue.W;
+	}
+	return ColorValue;
+}
diff --git a/Engine/Source/Editor/DetailCustomizations/Private/Vector4StructCustomization.h b/Engine/Source/Editor/DetailCustomizations/Private/Vector4StructCustomization.h
--- a/Engine/Source/Editor/DetailCustomizations/Private/Vector4StructCustomization.h
+++ b/Engine/Source/Editor/DetailCustomizations/Private/Vector4StructCustomization.h
@@ -8,6 +8,7 @@
 
 /**
  * Customizes FVector4 structs.
+ * We override Vector4 because the color grading controls are made with vector4
  */
 class FVector4StructCustomization
 	: public FMathStructCustomization
@@ -26,5 +27,8 @@ public:
 	
 	// We specialize the detail display of color grading vector property. The color grading mode is specified inside the metadata of the UProperty
 	TSharedPtr<FColorGradingVectorCustomization> ColorGradingVectorCustomization;
+
+private:
+	FLinearColor OnGetHeaderColorBlock(TWeakPtr<IPropertyHandle> WeakHandlePtr) const;
 };
 
diff --git a/Engine/Source/Editor/DistCurveEditor/Public/DistCurveEditorModule.h b/Engine/Source/Editor/DistCurveEditor/Public/DistCurveEditorModule.h
--- a/Engine/Source/Editor/DistCurveEditor/Public/DistCurveEditorModule.h
+++ b/Engine/Source/Editor/DistCurveEditor/Public/DistCurveEditorModule.h
@@ -1,7 +1,6 @@
 // Copyright 1998-2017 Epic Games, Inc. All Rights Reserved.
 
-#ifndef __DistCurveEditorModule_h__
-#define __DistCurveEditorModule_h__
+#pragma once
 
 #include "CoreMinimal.h"
 #include "Modules/ModuleInterface.h"
@@ -23,5 +22,3 @@ public:
 	virtual TSharedRef<IDistributionCurveEditor> CreateCurveEditorWidget(UInterpCurveEdSetup* EdSetup, FCurveEdNotifyInterface* NotifyObject) = 0;
 	virtual TSharedRef<IDistributionCurveEditor> CreateCurveEditorWidget(UInterpCurveEdSetup* EdSetup, FCurveEdNotifyInterface* NotifyObject, IDistributionCurveEditor::FCurveEdOptions Options) = 0;
 };
-
-#endif // __DistCurveEditorModule_h__
diff --git a/Engine/Source/Editor/EditorStyle/Private/SlateEditorStyle.cpp b/Engine/Source/Editor/EditorStyle/Private/SlateEditorStyle.cpp
--- a/Engine/Source/Editor/EditorStyle/Private/SlateEditorStyle.cpp
+++ b/Engine/Source/Editor/EditorStyle/Private/SlateEditorStyle.cpp
@@ -5854,6 +5854,7 @@ void FSlateEditorStyle::FStyle::SetupClassIconsAndThumbnails()
 			TEXT("Skeleton"),
 			TEXT("SlateBrushAsset"),
 			TEXT("SlateWidgetStyleAsset"),
+			TEXT("StringTable"),
 			TEXT("SoundAttenuation"),
 			TEXT("SoundClass"),
 			TEXT("SoundConcurrency"),
@@ -6234,6 +6235,13 @@ void FSlateEditorStyle::FStyle::SetupLandscapeEditorStyle()
 		Set("LandscapeEditor.Target_Create",     new IMAGE_BRUSH("Icons/icon_Landscape_Target_Create_12x", Icon12x12));
 		Set("LandscapeEditor.Target_MakePublic", new IMAGE_BRUSH("Icons/assign_right_12x",                 Icon12x12));
 		Set("LandscapeEditor.Target_Delete",     new IMAGE_BRUSH("Icons/Cross_12x",                        Icon12x12));
+
+		Set("LandscapeEditor.Target_DisplayOrder.Default", new IMAGE_BRUSH("Icons/icon_landscape_sort_base", Icon16x16));
+		Set("LandscapeEditor.Target_DisplayOrder.Alphabetical", new IMAGE_BRUSH("Icons/icon_landscape_sort_alphabetical", Icon16x16));
+		Set("LandscapeEditor.Target_DisplayOrder.Custom", new IMAGE_BRUSH("Icons/icon_landscape_sort_custom", Icon16x16));
+
+		Set("LandscapeEditor.TargetList.DropZone.Above", new BOX_BRUSH("Common/VerticalBoxDropZoneIndicator_Above", FMargin(10.0f / 16.0f, 10.0f / 16.0f, 0, 0), SelectionColor_Subdued));
+		Set("LandscapeEditor.TargetList.DropZone.Below", new BOX_BRUSH("Common/VerticalBoxDropZoneIndicator_Below", FMargin(10.0f / 16.0f, 0, 0, 10.0f / 16.0f), SelectionColor_Subdued));
 	}
 #endif
 }
@@ -6331,6 +6339,11 @@ void FSlateEditorStyle::FStyle::SetupToolkitStyles()
 		);
 	}
 
+	// StringTable Editor
+	{
+		Set("StringTableEditor.Tabs.Properties", new IMAGE_BRUSH("/Icons/icon_tab_SelectionDetails_16x", Icon16x16));
+	}
+
 	// DestructibleMesh Editor
 	{
 		Set( "DestructibleMeshEditor.Tabs.Properties", new IMAGE_BRUSH( "/Icons/icon_tab_SelectionDetails_16x", Icon16x16 ) );
