Automation - Fixing the screenshot browser to refresh when changing the export directory for screenshots.

diff --git a/Engine/Source/Developer/ScreenShotComparison/Private/Widgets/SScreenShotBrowser.cpp b/Engine/Source/Developer/ScreenShotComparison/Private/Widgets/SScreenShotBrowser.cpp
--- a/Engine/Source/Developer/ScreenShotComparison/Private/Widgets/SScreenShotBrowser.cpp
+++ b/Engine/Source/Developer/ScreenShotComparison/Private/Widgets/SScreenShotBrowser.cpp
@@ -127,6 +127,8 @@ void SScreenShotBrowser::OnDirectoryChanged(const FString& Directory)
 	ComparisonRoot = Directory;
 
 	RefreshDirectoryWatcher();
+
+	bReportsChanged = true;
 }
 
 void SScreenShotBrowser::RefreshDirectoryWatcher()
