Copying //UE4/Dev-Automation to //UE4/Dev-Main (Source: //UE4/Dev-Automation @ 3723189)

============================
  MAJOR FEATURES & CHANGES
============================

Change 3716887 by [REDACTED_PERSON]

	[Duplicate] Added "show delta" checkbox to automated screenshot comparison viewer.

Change 3720970 by [REDACTED_PERSON]

	Added more time for UDP serialization test.

Change 3722873 by [REDACTED_PERSON]

	Updated screenshot comparison to handle new folder layout, can no longer assume the root folder. Was deleting pending AddNew/Replace comparison results when any were updated.

Change 3720019 by [REDACTED_PERSON]

	Swapping automation screenshot test data from "Map_Test/Platform_RHI" to "Map/Test/Platform/RHI" as this allows platform folders and platform-specific tests to obey privacy rules.

Change 3720989 by [REDACTED_PERSON]

	[Duplicate] Allow platforms to override the default ICU timezone calculation

diff --git a/Engine/Plugins/Messaging/UdpMessaging/Source/UdpMessaging/Private/Tests/UdpSerializeMessageTaskTest.cpp b/Engine/Plugins/Messaging/UdpMessaging/Source/UdpMessaging/Private/Tests/UdpSerializeMessageTaskTest.cpp
--- a/Engine/Plugins/Messaging/UdpMessaging/Source/UdpMessaging/Private/Tests/UdpSerializeMessageTaskTest.cpp
+++ b/Engine/Plugins/Messaging/UdpMessaging/Source/UdpMessaging/Private/Tests/UdpSerializeMessageTaskTest.cpp
@@ -17,7 +17,7 @@ IMPLEMENT_SIMPLE_AUTOMATION_TEST(FUdpSerializeMessageTaskTest, "System.Core.Mess
 namespace UdpSerializeMessageTaskTest
 {
 	const int32 NumMessages = 100000;
-	const FTimespan MaxWaitTime(0, 0, 5);
+	const FTimespan MaxWaitTime(0, 0, 10);
 
 	int32 CompletedMessages = 0;
 	int32 FailedMessages = 0;
diff --git a/Engine/Source/Developer/FunctionalTesting/Private/AutomationBlueprintFunctionLibrary.cpp b/Engine/Source/Developer/FunctionalTesting/Private/AutomationBlueprintFunctionLibrary.cpp
--- a/Engine/Source/Developer/FunctionalTesting/Private/AutomationBlueprintFunctionLibrary.cpp
+++ b/Engine/Source/Developer/FunctionalTesting/Private/AutomationBlueprintFunctionLibrary.cpp
@@ -47,7 +47,7 @@ static TAutoConsoleVariable<int32> CVarAutomationScreenshotResolutionHeight(
 	TEXT("The height of automation screenshots."),
 	ECVF_Default);
 
-static TAutoConsoleVariable<int32> CVarMaxContactOffset(
+static TAutoConsoleVariable<int32> CVarAutomationTestUseRelaxedPhysicsStabilityThresholds(
 	TEXT("AutomationTestUseRelaxedPhysicsStabilityThresholds"),
 	0,
 	TEXT("Relaxes the thresold of automated physics stability tests. Default: Off"),
diff --git a/Engine/Source/Developer/ScreenShotComparison/Private/Models/ScreenComparisonModel.cpp b/Engine/Source/Developer/ScreenShotComparison/Private/Models/ScreenComparisonModel.cpp
--- a/Engine/Source/Developer/ScreenShotComparison/Private/Models/ScreenComparisonModel.cpp
+++ b/Engine/Source/Developer/ScreenShotComparison/Private/Models/ScreenComparisonModel.cpp
@@ -32,16 +32,25 @@ void FScreenComparisonModel::Complete()
 	FString RelativeReportFolder = Report.ReportFolder;
 	if ( FPaths::MakePathRelativeTo(RelativeReportFolder, *Report.ReportRootDirectory) )
 	{
+		// Find test folder, immediate sub-folder of map
+		// e.g. Map/Test/Platform/RHI -> Map/Test 
 		for (;;)
 		{
 			FString ParentFolder = FPaths::GetPath(RelativeReportFolder);
-			if ( ParentFolder.IsEmpty() )
+
+			int32 SubfolderPos = INDEX_NONE;
+			bool bContainsSubfolder = false;
+			bContainsSubfolder |= ParentFolder.FindChar('\\', SubfolderPos);
+			bContainsSubfolder |= ParentFolder.FindChar('/', SubfolderPos);
+
+			if (ParentFolder.IsEmpty() || !bContainsSubfolder)
 			{
 				break;
 			}
 			RelativeReportFolder = ParentFolder;
 		}
 
+		// Delete test folder
 		FString ReportTopFolder = Report.ReportRootDirectory / RelativeReportFolder;
 		if ( IFileManager::Get().DeleteDirectory(*ReportTopFolder, false, true) )
 		{
diff --git a/Engine/Source/Developer/ScreenShotComparison/Private/Widgets/SScreenComparisonRow.cpp b/Engine/Source/Developer/ScreenShotComparison/Private/Widgets/SScreenComparisonRow.cpp
--- a/Engine/Source/Developer/ScreenShotComparison/Private/Widgets/SScreenComparisonRow.cpp
+++ b/Engine/Source/Developer/ScreenShotComparison/Private/Widgets/SScreenComparisonRow.cpp
@@ -14,6 +14,7 @@
 #include "Widgets/Layout/SScaleBox.h"
 #include "Widgets/Input/SButton.h"
 #include "Widgets/Input/SSlider.h"
+#include "Widgets/Input/SCheckBox.h"
 #include "HAL/FileManager.h"
 #include "Misc/FileHelper.h"
 #include "IImageWrapperModule.h"
@@ -30,14 +31,14 @@ public:
 	SLATE_BEGIN_ARGS(SImageComparison) {}
 		SLATE_ARGUMENT(TSharedPtr<FSlateDynamicImageBrush>, BaseImage)
 		SLATE_ARGUMENT(TSharedPtr<FSlateDynamicImageBrush>, ModifiedImage)
+		SLATE_ARGUMENT(TSharedPtr<FSlateDynamicImageBrush>, DeltaImage)
 	SLATE_END_ARGS()
 
 	void Construct(const FArguments& InArgs)
 	{
 		BaseImage = InArgs._BaseImage;
 		ModifiedImage = InArgs._ModifiedImage;
-
-		//const int32 MaxHeight = FSlateApplication::Get().GetPreferredWorkArea().GetSize().Y * 0.80;
+		DeltaImage = InArgs._DeltaImage;
 
 		ChildSlot
 		[
@@ -63,6 +64,13 @@ public:
 						.Image(ModifiedImage.Get())
 						.ColorAndOpacity(this, &SImageComparison::GetModifiedOpacity)
 					]
+
+					+ SOverlay::Slot()
+					[
+						SNew(SImage)
+						.Image(DeltaImage.Get())
+						.ColorAndOpacity(this, &SImageComparison::GetDeltaOpacity)
+					]
 				]
 			]
 			
@@ -72,6 +80,7 @@ public:
 			[
 				SNew(SHorizontalBox)
 				
+				// Comparison slider
 				+ SHorizontalBox::Slot()
 				.AutoWidth()
 				[
@@ -96,6 +105,30 @@ public:
 					SNew(STextBlock)
 					.Text(LOCTEXT("Incoming", "Incoming"))
 				]
+
+				// Delta checkbox
+				+SHorizontalBox::Slot()
+				.AutoWidth()
+				[
+					SNew(SBox)
+					.MinDesiredWidth(50)
+				]
+
+				+SHorizontalBox::Slot()
+				.AutoWidth()
+				[
+					SNew(SBox)
+					[
+						SAssignNew(DeltaCheckbox, SCheckBox)
+					]
+				]
+
+				+ SHorizontalBox::Slot()
+				.AutoWidth()
+				[
+					SNew(STextBlock)
+					.Text(LOCTEXT("ShowDelta", "Show Delta"))
+				]
 			]
 		];
 	}
@@ -105,11 +138,18 @@ public:
 		return FLinearColor(1, 1, 1, OpacitySlider->GetValue());
 	}
 
+	FSlateColor GetDeltaOpacity() const
+	{
+		return FLinearColor(1, 1, 1, DeltaCheckbox->IsChecked() ? 1 : 0);
+	}
+
 private:
 	TSharedPtr<FSlateDynamicImageBrush> BaseImage;
 	TSharedPtr<FSlateDynamicImageBrush> ModifiedImage;
+	TSharedPtr<FSlateDynamicImageBrush> DeltaImage;
 
 	TSharedPtr<SSlider> OpacitySlider;
+	TSharedPtr<SCheckBox> DeltaCheckbox;
 };
 
 
@@ -284,7 +324,7 @@ TSharedRef<SWidget> SScreenComparisonRow::BuildComparisonPreview()
 						.AutoWidth()
 						.Padding(4.0f, 4.0f)
 						[
-							SNew(SAsyncImage)
+							SAssignNew(DeltaImageWidget, SAsyncImage)
 							.ImageFilePath(DeltaFile)
 						]
 
@@ -371,8 +411,9 @@ FReply SScreenComparisonRow::OnCompareImages(const FGeometry& InGeometry, const
 {
 	TSharedPtr<FSlateDynamicImageBrush> ApprovedImage = ApprovedImageWidget->GetDynamicBrush();
 	TSharedPtr<FSlateDynamicImageBrush> UnapprovedImage = UnapprovedImageWidget->GetDynamicBrush();
+	TSharedPtr<FSlateDynamicImageBrush> DeltaImage = DeltaImageWidget->GetDynamicBrush();
 
-	if ( ApprovedImage.IsValid() && UnapprovedImage.IsValid() )
+	if ( ApprovedImage.IsValid() && UnapprovedImage.IsValid() && DeltaImage.IsValid()  )
 	{
 		TSharedRef<SWindow> ParentWindow = FSlateApplication::Get().FindWidgetWindow(AsShared()).ToSharedRef();
 
@@ -391,6 +432,7 @@ FReply SScreenComparisonRow::OnCompareImages(const FGeometry& InGeometry, const
 				SNew(SImageComparison)
 				.BaseImage(ApprovedImage)
 				.ModifiedImage(UnapprovedImage)
+				.DeltaImage(DeltaImage)
 			];
 
 		FSlateApplication::Get().AddWindowAsNativeChild(PopupWindow, ParentWindow, true);
diff --git a/Engine/Source/Developer/ScreenShotComparison/Private/Widgets/SScreenComparisonRow.h b/Engine/Source/Developer/ScreenShotComparison/Private/Widgets/SScreenComparisonRow.h
--- a/Engine/Source/Developer/ScreenShotComparison/Private/Widgets/SScreenComparisonRow.h
+++ b/Engine/Source/Developer/ScreenShotComparison/Private/Widgets/SScreenComparisonRow.h
@@ -74,5 +74,6 @@ private:
 	FIntPoint CachedActualImageSize;
 
 	TSharedPtr<SAsyncImage> ApprovedImageWidget;
+	TSharedPtr<SAsyncImage> DeltaImageWidget;
 	TSharedPtr<SAsyncImage> UnapprovedImageWidget;
 };
diff --git a/Engine/Source/Runtime/Core/Private/GenericPlatform/GenericPlatformMisc.cpp b/Engine/Source/Runtime/Core/Private/GenericPlatform/GenericPlatformMisc.cpp
--- a/Engine/Source/Runtime/Core/Private/GenericPlatform/GenericPlatformMisc.cpp
+++ b/Engine/Source/Runtime/Core/Private/GenericPlatform/GenericPlatformMisc.cpp
@@ -917,6 +917,12 @@ FString FGenericPlatformMisc::GetDefaultLocale()
 #endif
 }
 
+FString FGenericPlatformMisc::GetTimeZoneId()
+{
+	// ICU will calculate this correctly for most platforms (if enabled)
+	return FString();
+}
+
 FText FGenericPlatformMisc::GetFileManagerName()
 {
 	return NSLOCTEXT("GenericPlatform", "FileManagerName", "File Manager");
diff --git a/Engine/Source/Runtime/Core/Private/Internationalization/ICUInternationalization.cpp b/Engine/Source/Runtime/Core/Private/Internationalization/ICUInternationalization.cpp
--- a/Engine/Source/Runtime/Core/Private/Internationalization/ICUInternationalization.cpp
+++ b/Engine/Source/Runtime/Core/Private/Internationalization/ICUInternationalization.cpp
@@ -12,6 +12,7 @@
 
 #if UE_ENABLE_ICU
 
+#include "Internationalization/ICUUtilities.h"
 #include "Internationalization/ICUBreakIterator.h"
 THIRD_PARTY_INCLUDES_START
 	#include <unicode/locid.h>
@@ -150,6 +151,7 @@ bool FICUInternationalization::Initialize()
 	I18N->CurrentLanguage = I18N->DefaultLanguage;
 	I18N->CurrentLocale = I18N->DefaultLocale;
 
+	InitializeTimeZone();
 	InitializeInvariantGregorianCalendar();
 
 	return U_SUCCESS(ICUStatus) ? true : false;
@@ -681,6 +683,19 @@ FCulturePtr FICUInternationalization::FindOrMakeCulture(const FString& Name, con
 	return NewCulture;
 }
 
+void FICUInternationalization::InitializeTimeZone()
+{
+	const FString TimeZoneId = FPlatformMisc::GetTimeZoneId();
+
+	icu::TimeZone* ICUDefaultTz = TimeZoneId.IsEmpty() ? icu::TimeZone::createDefault() : icu::TimeZone::createTimeZone(ICUUtilities::ConvertString(TimeZoneId));
+	icu::TimeZone::adoptDefault(ICUDefaultTz);
+
+	const int32 DefaultTzOffsetMinutes = ICUDefaultTz->getRawOffset() / 60000;
+	const int32 RawOffsetHours = DefaultTzOffsetMinutes / 60;
+	const int32 RawOffsetMinutes = DefaultTzOffsetMinutes % 60;
+	UE_LOG(LogICUInternationalization, Display, TEXT("ICU TimeZone Detection - Raw Offset: %+d:%02d, Platform Override: '%s'"), RawOffsetHours, RawOffsetMinutes, *TimeZoneId);
+}
+
 void FICUInternationalization::InitializeInvariantGregorianCalendar()
 {
 	UErrorCode ICUStatus = U_ZERO_ERROR;
diff --git a/Engine/Source/Runtime/Core/Private/Internationalization/ICUInternationalization.h b/Engine/Source/Runtime/Core/Private/Internationalization/ICUInternationalization.h
--- a/Engine/Source/Runtime/Core/Private/Internationalization/ICUInternationalization.h
+++ b/Engine/Source/Runtime/Core/Private/Internationalization/ICUInternationalization.h
@@ -56,6 +56,7 @@ private:
 	enum class EAllowDefaultCultureFallback : uint8 { No, Yes, };
 	FCulturePtr FindOrMakeCulture(const FString& Name, const EAllowDefaultCultureFallback AllowDefaultFallback);
 
+	void InitializeTimeZone();
 	void InitializeInvariantGregorianCalendar();
 
 private:
diff --git a/Engine/Source/Runtime/Core/Public/GenericPlatform/GenericPlatformMisc.h b/Engine/Source/Runtime/Core/Public/GenericPlatform/GenericPlatformMisc.h
--- a/Engine/Source/Runtime/Core/Public/GenericPlatform/GenericPlatformMisc.h
+++ b/Engine/Source/Runtime/Core/Public/GenericPlatform/GenericPlatformMisc.h
@@ -781,6 +781,12 @@ public:
 	 */
 	static FString GetDefaultLocale();
 
+	/**
+	 * Get the timezone identifier for this platform, or an empty string if the default timezone calculation will work.
+	 * @note This should return either an Olson timezone (eg, "America/Los_Angeles") or an offset from GMT/UTC (eg, "GMT-8:00").
+	 */
+	static FString GetTimeZoneId();
+
 	/**
 	 *	Platform-specific Exec function
 	 *
diff --git a/Engine/Source/Runtime/Engine/Private/Tests/AutomationCommon.cpp b/Engine/Source/Runtime/Engine/Private/Tests/AutomationCommon.cpp
--- a/Engine/Source/Runtime/Engine/Private/Tests/AutomationCommon.cpp
+++ b/Engine/Source/Runtime/Engine/Private/Tests/AutomationCommon.cpp
@@ -85,7 +85,7 @@ namespace AutomationCommon
 	void GetScreenshotPath(const FString& TestName, FString& OutScreenshotName)
 	{
 		FString PathName = FPaths::AutomationDir() + TestName / FPlatformProperties::IniPlatformName();
-		PathName = PathName + TEXT("_") + GetRenderDetailsString();
+		PathName = PathName + TEXT("/") + GetRenderDetailsString();
 
 		FPaths::MakePathRelativeTo(PathName, *FPaths::RootDir());
 
@@ -129,7 +129,7 @@ namespace AutomationCommon
 		// TBD - 
 		// Device's native resolution (we want to use a hardware dump of the frontbuffer at the native resolution so we compare what we actually output rather than what we think we rendered)
 
-		const FString MapAndTest = MapOrContext + TEXT("_") + TestName;
+		const FString MapAndTest = MapOrContext + TEXT("/") + TestName;
 		AutomationCommon::GetScreenshotPath(MapAndTest, Data.Path);
 
 		return Data;
