ScreenshotBrowser: Adding support for enlarging New screenshots with a click, to better preview them before approval

diff --git a/Engine/Source/Developer/ScreenShotComparison/Private/Widgets/SScreenComparisonRow.cpp b/Engine/Source/Developer/ScreenShotComparison/Private/Widgets/SScreenComparisonRow.cpp
--- a/Engine/Source/Developer/ScreenShotComparison/Private/Widgets/SScreenComparisonRow.cpp
+++ b/Engine/Source/Developer/ScreenShotComparison/Private/Widgets/SScreenComparisonRow.cpp
@@ -302,9 +302,14 @@ TSharedRef<SWidget> SScreenComparisonRow::BuildAddedView()
 					.AutoWidth()
 					.Padding(4.0f, 4.0f)
 					[
-						SNew(SAsyncImage)
-						.ImageFilePath(IncomingFile)
-						//.OnMouseButtonDown(this, &SScreenComparisonRow::OnImageClicked, UnapprovedBrush)
+						SNew(SBorder)
+						.BorderImage(nullptr)
+						.OnMouseButtonDown(this, &SScreenComparisonRow::OnCompareNewImage)
+						[
+							SAssignNew(UnapprovedImageWidget, SAsyncImage)
+							.ImageFilePath(IncomingFile)
+						]
+						
 					]
 				]
 			]
@@ -489,6 +494,42 @@ FReply SScreenComparisonRow::OnCompareImages(const FGeometry& InGeometry, const
 	return FReply::Handled();
 }
 
+FReply SScreenComparisonRow::OnCompareNewImage(const FGeometry& InGeometry, const FPointerEvent& InEvent)
+{
+	const FImageComparisonResult& ComparisonResult = Model->Report.Comparison;
+	FString IncomingFile = Model->Report.ReportFolder / ComparisonResult.ReportIncomingFile;
+
+	TSharedPtr<FSlateDynamicImageBrush> UnapprovedImage = UnapprovedImageWidget->GetDynamicBrush();
+
+	if (UnapprovedImage.IsValid())
+	{
+		TSharedRef<SWindow> ParentWindow = FSlateApplication::Get().FindWidgetWindow(AsShared()).ToSharedRef();
+
+		TSharedRef<SWindow> PopupWindow = SNew(SWindow)
+			.IsPopupWindow(false)
+			.ClientSize(FVector2D(1280,720))
+			.SizingRule(ESizingRule::UserSized)
+			.AutoCenter(EAutoCenter::PreferredWorkArea)
+			.SupportsMaximize(true)
+			.SupportsMinimize(true)
+			.FocusWhenFirstShown(true)
+			.ActivationPolicy(EWindowActivationPolicy::Always)
+			.Content()
+			[
+				SNew(SScaleBox)
+				.Stretch(EStretch::ScaleToFit)
+				[
+					SNew(SImage)
+					.Image(UnapprovedImage.Get())
+				]
+			];
+
+		FSlateApplication::Get().AddWindowAsNativeChild(PopupWindow, ParentWindow, true);
+	}
+	
+	return FReply::Handled();
+}
+
 FReply SScreenComparisonRow::OnImageClicked(const FGeometry& InGeometry, const FPointerEvent& InEvent, TSharedPtr<FSlateDynamicImageBrush> Image)
 {
 	TSharedRef<SWindow> ParentWindow = FSlateApplication::Get().FindWidgetWindow(AsShared()).ToSharedRef();
diff --git a/Engine/Source/Developer/ScreenShotComparison/Private/Widgets/SScreenComparisonRow.h b/Engine/Source/Developer/ScreenShotComparison/Private/Widgets/SScreenComparisonRow.h
--- a/Engine/Source/Developer/ScreenShotComparison/Private/Widgets/SScreenComparisonRow.h
+++ b/Engine/Source/Developer/ScreenShotComparison/Private/Widgets/SScreenComparisonRow.h
@@ -62,6 +62,7 @@ private:
 	FReply AddAlternative();
 
 	FReply OnCompareImages(const FGeometry& InGeometry, const FPointerEvent& InEvent);
+	FReply OnCompareNewImage(const FGeometry& InGeometry, const FPointerEvent& InEvent);
 	FReply OnImageClicked(const FGeometry& InGeometry, const FPointerEvent& InEvent, TSharedPtr<FSlateDynamicImageBrush> Image);
 
 private:
