ScreenshotBrowser: Adding Filters for displayed screenshot reports based on:
 - Screenshot Name (case-insensitive string)
 - Fails (checkbox to show or hide fails)
 - New Screenshots (checkbox to show or hide new screenshots)

diff --git a/Engine/Source/Developer/ScreenShotComparison/Private/Widgets/SScreenShotBrowser.cpp b/Engine/Source/Developer/ScreenShotComparison/Private/Widgets/SScreenShotBrowser.cpp
--- a/Engine/Source/Developer/ScreenShotComparison/Private/Widgets/SScreenShotBrowser.cpp
+++ b/Engine/Source/Developer/ScreenShotComparison/Private/Widgets/SScreenShotBrowser.cpp
@@ -10,6 +10,8 @@
 #include "Widgets/Input/SButton.h"
 #include "Widgets/Input/SDirectoryPicker.h"
 #include "Widgets/SScreenComparisonRow.h"
+#include "Widgets/Input/SCheckBox.h"
+#include "Widgets/Input/SEditableTextBox.h"
 #include "Models/ScreenComparisonModel.h"
 #include "Misc/FeedbackContext.h"
 #include "EditorStyleSet.h"
@@ -28,6 +30,9 @@ void SScreenShotBrowser::Construct( const FArguments& InArgs,  IScreenShotManage
 	ScreenShotManager = InScreenShotManager;
 	ComparisonRoot = FPaths::ConvertRelativePathToFull(FPaths::ProjectSavedDir() / TEXT("Automation/Comparisons"));
 	bReportsChanged = true;
+	bDisplayingError = true;
+	bDisplayingWarning = true;
+	ReportFilterString = FString();
 
 	FModuleManager::Get().LoadModuleChecked(FName("ImageWrapper"));
 
@@ -48,6 +53,75 @@ void SScreenShotBrowser::Construct( const FArguments& InArgs,  IScreenShotManage
 				.OnDirectoryChanged(this, &SScreenShotBrowser::OnDirectoryChanged)
 			]
 
+			+ SHorizontalBox::Slot()
+			.FillWidth(1.0f)
+			.HAlign(HAlign_Left)
+			.Padding(10.0f, 0.0f)
+			[
+				SNew(SHorizontalBox)
+
+				+ SHorizontalBox::Slot()
+				.FillWidth(1.0f)
+				.Padding(2.0f, 0.0f)
+				[
+					SNew(SEditableTextBox)
+					.HintText(LOCTEXT("ScreenshotFilterHint", "Filter Reports"))
+					.Text(FText::FromString(ReportFilterString))
+					.OnTextCommitted(this, &SScreenShotBrowser::OnFilterStringCommitted)
+					
+				]
+
+				+ SHorizontalBox::Slot()
+				.AutoWidth()
+				.Padding(2.0f, 0.0f)
+				[
+					SNew(SHorizontalBox)
+					
+					+ SHorizontalBox::Slot()
+					.AutoWidth()
+					[
+						SNew(SCheckBox)
+						.HAlign(HAlign_Center)
+						.IsChecked( bDisplayingError ? ECheckBoxState::Checked : ECheckBoxState::Unchecked)
+						.OnCheckStateChanged(this, &SScreenShotBrowser::DisplayError_OnCheckStateChanged)
+					]
+
+					+ SHorizontalBox::Slot()
+					.AutoWidth()
+					.VAlign(VAlign_Center)
+					[
+						SNew(STextBlock)
+						.AutoWrapText(true)
+						.Text(LOCTEXT("DisplayErrors", "Show Fails"))
+					]
+				]
+
+				+ SHorizontalBox::Slot()
+				.AutoWidth()
+				.Padding(2.0f, 0.0f)
+				[
+					SNew(SHorizontalBox)
+					
+					+ SHorizontalBox::Slot()
+					.AutoWidth()
+					[
+						SNew(SCheckBox)
+						.HAlign(HAlign_Center)
+						.IsChecked(bDisplayingWarning ? ECheckBoxState::Checked : ECheckBoxState::Unchecked)
+						.OnCheckStateChanged(this, &SScreenShotBrowser::DisplayWarning_OnCheckStateChanged)
+					]
+
+					+ SHorizontalBox::Slot()
+					.AutoWidth()
+					.VAlign(VAlign_Center)
+					[
+						SNew(STextBlock)
+						.AutoWrapText(true)
+						.Text(LOCTEXT("DisplayWarnings", "Show New"))
+					]
+				]
+			]
+
 			+ SHorizontalBox::Slot()
 			.FillWidth(1.0f)
 			.HAlign(HAlign_Right)
@@ -245,6 +319,28 @@ TSharedRef<ITableRow> SScreenShotBrowser::OnGenerateWidgetForScreenResults(TShar
 		.ComparisonResult(InItem);
 }
 
+void SScreenShotBrowser::DisplayError_OnCheckStateChanged(ECheckBoxState NewRadioState)
+{
+	bDisplayingError = (NewRadioState == ECheckBoxState::Checked);
+	bReportsChanged = true;
+}
+
+void SScreenShotBrowser::DisplayWarning_OnCheckStateChanged(ECheckBoxState NewRadioState)
+{
+	bDisplayingWarning = (NewRadioState == ECheckBoxState::Checked);
+	bReportsChanged = true;
+}
+
+void SScreenShotBrowser::OnFilterStringCommitted(const FText& InText, ETextCommit::Type InCommitType)
+{
+	FString InString = InText.ToString();
+	if (ReportFilterString.Compare(InString, ESearchCase::IgnoreCase) != 0)
+	{
+		ReportFilterString = InString;
+		bReportsChanged = true;
+	}
+}
+
 void SScreenShotBrowser::RebuildTree()
 {
 	bReportsChanged = false;
@@ -257,6 +353,18 @@ void SScreenShotBrowser::RebuildTree()
 
 		for ( const FComparisonReport& Report : CurrentReports )
 		{
+			if (Report.Comparison.IsNew() && !bDisplayingWarning)
+			{
+				continue;
+			}
+			if (!Report.Comparison.IsNew() && !Report.Comparison.AreSimilar() && !bDisplayingError)
+			{
+				continue;
+			}
+			if (ReportFilterString.Len() && !Report.ReportFolder.Contains(ReportFilterString, ESearchCase::IgnoreCase))
+			{
+				continue;
+			}
 			TSharedPtr<FScreenComparisonModel> Model = MakeShared<FScreenComparisonModel>(Report);
 			Model->OnComplete.AddLambda([this, Model] () {
 				ComparisonList.Remove(Model);
diff --git a/Engine/Source/Developer/ScreenShotComparison/Private/Widgets/SScreenShotBrowser.h b/Engine/Source/Developer/ScreenShotComparison/Private/Widgets/SScreenShotBrowser.h
--- a/Engine/Source/Developer/ScreenShotComparison/Private/Widgets/SScreenShotBrowser.h
+++ b/Engine/Source/Developer/ScreenShotComparison/Private/Widgets/SScreenShotBrowser.h
@@ -43,6 +43,10 @@ public:
 
 	TSharedRef<ITableRow> OnGenerateWidgetForScreenResults(TSharedPtr<FScreenComparisonModel> InItem, const TSharedRef<STableViewBase>& OwnerTable);
 
+	void DisplayError_OnCheckStateChanged(ECheckBoxState NewRadioState);
+	void DisplayWarning_OnCheckStateChanged(ECheckBoxState NewRadioState);
+	void OnFilterStringCommitted(const FText& InText, ETextCommit::Type InCommitType);
+
 private:
 
 	void OnDirectoryChanged(const FString& Directory);
@@ -80,4 +84,13 @@ private:
 
 	/**  */
 	bool bReportsChanged;
+
+	/** Whether or not we're currently displaying errors */
+	bool bDisplayingError;
+
+	/** Whether or not we're currently displaying warnings*/
+	bool bDisplayingWarning;
+
+	/** Filter string for our reports, if any */
+	FString ReportFilterString;
 };
