Adding 'delete' button on individual comparisons in the Screenshot Comparison Browser

diff --git a/Engine/Source/Developer/ScreenShotComparison/Private/Widgets/SScreenComparisonRow.cpp b/Engine/Source/Developer/ScreenShotComparison/Private/Widgets/SScreenComparisonRow.cpp
--- a/Engine/Source/Developer/ScreenShotComparison/Private/Widgets/SScreenComparisonRow.cpp
+++ b/Engine/Source/Developer/ScreenShotComparison/Private/Widgets/SScreenComparisonRow.cpp
@@ -257,6 +257,16 @@ TSharedRef<SWidget> SScreenComparisonRow::GenerateWidgetForColumn(const FName& C
 						.Text(LOCTEXT("AddAlternative", "Add As Alternative"))
 						.OnClicked(this, &SScreenComparisonRow::AddAlternative)
 					]
+
+					+ SHorizontalBox::Slot()
+					.Padding(10, 0, 0, 0)
+					.AutoWidth()
+					[
+						SNew(SButton)
+						.IsEnabled(true)
+						.Text(LOCTEXT("Delete", "Delete"))
+						.OnClicked(this, &SScreenComparisonRow::Remove)
+					]
 				];
 		}
 	}
@@ -319,10 +329,26 @@ TSharedRef<SWidget> SScreenComparisonRow::BuildAddedView()
 		.AutoHeight()
 		.HAlign(HAlign_Center)
 		[
-			SNew(SButton)
-			.IsEnabled(this, &SScreenComparisonRow::CanUseSourceControl)
-			.Text(LOCTEXT("AddNew", "Add New!"))
-			.OnClicked(this, &SScreenComparisonRow::AddNew)
+			SNew(SHorizontalBox)
+
+			+ SHorizontalBox::Slot()
+			.AutoWidth()
+			[
+				SNew(SButton)
+				.IsEnabled(this, &SScreenComparisonRow::CanUseSourceControl)
+				.Text(LOCTEXT("AddNew", "Add New!"))
+				.OnClicked(this, &SScreenComparisonRow::AddNew)
+			]
+
+			+ SHorizontalBox::Slot()
+			.AutoWidth()
+			.Padding(10, 0, 0, 0)
+			[
+				SNew(SButton)
+				.IsEnabled(true)
+				.Text(LOCTEXT("Delete", "Delete"))
+				.OnClicked(this, &SScreenComparisonRow::Remove)
+			]
 		];
 }
 
@@ -460,6 +486,13 @@ FReply SScreenComparisonRow::AddAlternative()
 	return FReply::Handled();
 }
 
+FReply SScreenComparisonRow::Remove()
+{
+	Model->Complete();
+
+	return FReply::Handled();
+}
+
 FReply SScreenComparisonRow::OnCompareImages(const FGeometry& InGeometry, const FPointerEvent& InEvent)
 {
 	TSharedPtr<FSlateDynamicImageBrush> ApprovedImage = ApprovedImageWidget->GetDynamicBrush();
diff --git a/Engine/Source/Developer/ScreenShotComparison/Private/Widgets/SScreenComparisonRow.h b/Engine/Source/Developer/ScreenShotComparison/Private/Widgets/SScreenComparisonRow.h
--- a/Engine/Source/Developer/ScreenShotComparison/Private/Widgets/SScreenComparisonRow.h
+++ b/Engine/Source/Developer/ScreenShotComparison/Private/Widgets/SScreenComparisonRow.h
@@ -61,6 +61,8 @@ private:
 	bool CanAddAsAlternative() const;
 	FReply AddAlternative();
 
+	FReply Remove();
+
 	FReply OnCompareImages(const FGeometry& InGeometry, const FPointerEvent& InEvent);
 	FReply OnCompareNewImage(const FGeometry& InGeometry, const FPointerEvent& InEvent);
 	FReply OnImageClicked(const FGeometry& InGeometry, const FPointerEvent& InEvent, TSharedPtr<FSlateDynamicImageBrush> Image);
