Replacing ground truth images wasn't properly deleting previous results
* The revert source control operation would fail if the files weren't already checked out

diff --git a/Engine/Source/Developer/ScreenShotComparison/Private/Models/ScreenComparisonModel.cpp b/Engine/Source/Developer/ScreenShotComparison/Private/Models/ScreenComparisonModel.cpp
--- a/Engine/Source/Developer/ScreenShotComparison/Private/Models/ScreenComparisonModel.cpp
+++ b/Engine/Source/Developer/ScreenShotComparison/Private/Models/ScreenComparisonModel.cpp
@@ -156,24 +156,45 @@ bool FScreenComparisonModel::RemoveExistingApproved()
 
 	TArray<FString> SourceControlFiles;
 
-	bool bSuccess = false;
+	bool bSuccess = true;
 
 	IFileManager::Get().FindFilesRecursive(SourceControlFiles, *FPaths::GetPath(ApprovedFolder), TEXT("*.*"), true, false, false);
 
+	// Remove files from source control.
 	if (SourceControlFiles.Num())
 	{
 		ISourceControlProvider& SourceControlProvider = ISourceControlModule::Get().GetProvider();
-		if (SourceControlProvider.Execute(ISourceControlOperation::Create<FRevert>(), SourceControlFiles) != ECommandResult::Failed)
+
+		TArray<FSourceControlStateRef> SourceControlStates;
+		ECommandResult::Type Result = SourceControlProvider.GetState(SourceControlFiles, SourceControlStates, EStateCacheUsage::ForceUpdate);
+		if (Result == ECommandResult::Succeeded)
 		{
-			if (SourceControlProvider.Execute(ISourceControlOperation::Create<FDelete>(), SourceControlFiles) != ECommandResult::Failed)
+			TArray<FString> FilesToRevert;
+			TArray<FString> FilesToDelete;
+
+			for (const FSourceControlStateRef& SourceControlState : SourceControlStates)
 			{
-				for (const FString& File : SourceControlFiles)
+				// Added files must be reverted.
+				if (SourceControlState->IsAdded())
 				{
-					IFileManager::Get().Delete(*File, false, true, false);
+					FilesToRevert.Add(SourceControlState->GetFilename());
+				}
+				// Edited files must be reverted then deleted.
+				else if (SourceControlState->IsCheckedOut())
+				{
+					FilesToRevert.Add(SourceControlState->GetFilename());
+					FilesToDelete.Add(SourceControlState->GetFilename());
+				}
+				// Other source controlled files must be deleted.
+				else if (SourceControlState->IsSourceControlled() && !SourceControlState->IsDeleted())
+				{
+					FilesToDelete.Add(SourceControlState->GetFilename());
 				}
-
-				bSuccess = true;
 			}
+
+			const bool WasSuccessfullyReverted = FilesToRevert.Num() == 0 || SourceControlProvider.Execute(ISourceControlOperation::Create<FRevert>(), FilesToRevert) == ECommandResult::Succeeded;
+			const bool WasSuccessfullyDeleted = FilesToDelete.Num() == 0 || SourceControlProvider.Execute(ISourceControlOperation::Create<FDelete>(), FilesToDelete) == ECommandResult::Succeeded;
+			bSuccess = WasSuccessfullyReverted && WasSuccessfullyDeleted;
 		}
 	}
 
