Reorganize screenshot comparison input fields in Session Frontend

diff --git a/Engine/Source/Developer/ScreenShotComparison/Private/Widgets/SScreenShotBrowser.cpp b/Engine/Source/Developer/ScreenShotComparison/Private/Widgets/SScreenShotBrowser.cpp
--- a/Engine/Source/Developer/ScreenShotComparison/Private/Widgets/SScreenShotBrowser.cpp
+++ b/Engine/Source/Developer/ScreenShotComparison/Private/Widgets/SScreenShotBrowser.cpp
@@ -11,7 +11,7 @@
 #include "Widgets/Input/SDirectoryPicker.h"
 #include "Widgets/SScreenComparisonRow.h"
 #include "Widgets/Input/SCheckBox.h"
-#include "Widgets/Input/SEditableTextBox.h"
+#include "Widgets/Input/SSearchBox.h"
 #include "Models/ScreenComparisonModel.h"
 #include "Misc/FeedbackContext.h"
 #include "Styling/AppStyle.h"
@@ -47,7 +47,7 @@ void SScreenShotBrowser::Construct( const FArguments& InArgs,  IScreenShotManage
 	bDisplayingError = true;
 	bDisplayingNew = true;
 	ReportFilterString = FString();
-
+	
 	FModuleManager::Get().LoadModuleChecked(FName("ImageWrapper"));
 
 	ChildSlot
@@ -59,6 +59,96 @@ void SScreenShotBrowser::Construct( const FArguments& InArgs,  IScreenShotManage
 		[
 			SNew(SHorizontalBox)
 
+			+ SHorizontalBox::Slot()
+			.AutoWidth()
+			.Padding(2.f, 0.f)
+			[
+				SNew(SPositiveActionButton)
+				.Text(LOCTEXT("AddAllNewReports", "Add All New Reports"))
+				.ToolTipText(LOCTEXT("AddAllNewReportsTooltip", "Adds all new screenshots contained in the reports."))
+				.IsEnabled_Lambda([this]() -> bool
+					{
+						return ComparisonList.Num() > 0 && ISourceControlModule::Get().IsEnabled();
+					})
+				.OnClicked_Lambda([this]()
+					{
+						for (int Entry = ComparisonList.Num() - 1; Entry >= 0; --Entry)
+						{
+							TSharedPtr<FScreenComparisonModel> Model = ComparisonList[Entry];
+							const FImageComparisonResult& Comparison = Model->Report.GetComparisonResult();
+
+							if (CanAddNewReportResult(Comparison))
+							{
+								ComparisonList.RemoveAt(Entry);
+								Model->AddNew();
+
+								// Avoid thrashing P4
+								const float SleepBetweenOperations = 0.005f;
+								FPlatformProcess::Sleep(SleepBetweenOperations);
+							}
+						}
+
+						return FReply::Handled();
+					})
+			]
+
+			+ SHorizontalBox::Slot()
+			.AutoWidth()
+			.Padding(2.f, 0.f)
+			[
+				SNew(SNegativeActionButton)
+				.ActionButtonStyle(EActionButtonStyle::Warning)
+				.Text(LOCTEXT("ReplaceAllReports", "Replace All Reports"))
+				.ToolTipText(LOCTEXT("ReplaceAllReportsTooltip", "Replaces all screenshots containing a different result in the reports."))
+				.IsEnabled_Lambda([this]() -> bool
+					{
+						return ComparisonList.Num() > 0 && ISourceControlModule::Get().IsEnabled();
+					})
+				.OnClicked_Lambda([this]()
+					{
+						for (int Entry = ComparisonList.Num() - 1; Entry >= 0; --Entry)
+						{
+							TSharedPtr<FScreenComparisonModel> Model = ComparisonList[Entry];
+							const FImageComparisonResult& Comparison = Model->Report.GetComparisonResult();
+
+							if (!CanAddNewReportResult(Comparison))
+							{
+								ComparisonList.RemoveAt(Entry);
+								Model->Replace();
+
+								// Avoid thrashing P4
+								const float SleepBetweenOperations = 0.005f;
+								FPlatformProcess::Sleep(SleepBetweenOperations);
+							}
+						}
+
+						return FReply::Handled();
+					})
+			]
+
+			+ SHorizontalBox::Slot()
+			.AutoWidth()
+			.Padding(2.f, 0.f)
+			[
+				SNew(SNegativeActionButton)
+				.Text(LOCTEXT("DeleteAllReports", "Delete All Reports"))
+				.ToolTipText(LOCTEXT("DeleteAllReportsTooltip", "Deletes all the current reports.  Reports are not removed unless the user resolves them, \nso if you just want to reset the state of the reports, clear them here and then re-run the tests."))
+				.IsEnabled_Lambda([this]() -> bool
+					{
+						return ComparisonList.Num() > 0;
+					})
+				.OnClicked_Lambda([this]()
+					{
+						while (ComparisonList.Num() > 0)
+						{
+							TSharedPtr<FScreenComparisonModel> Model = ComparisonList.Pop();
+							Model->Complete(true);
+						}
+
+						return FReply::Handled();
+					})
+			]
+
 			+ SHorizontalBox::Slot()
 			.AutoWidth()
 			.Padding(2.0f, 0.0f)
@@ -67,202 +157,118 @@ void SScreenShotBrowser::Construct( const FArguments& InArgs,  IScreenShotManage
 				.Directory(ComparisonRoot)
 				.OnDirectoryChanged(this, &SScreenShotBrowser::OnDirectoryChanged)
 			]
+		]
+
+		+ SVerticalBox::Slot()
+			.AutoHeight()
+		[
+			SNew(SHorizontalBox)
 
 			+ SHorizontalBox::Slot()
 			.FillWidth(1.0f)
-			.HAlign(HAlign_Left)
-			.Padding(10.0f, 0.0f)
 			[
 				SNew(SHorizontalBox)
 
 				+ SHorizontalBox::Slot()
 				.FillWidth(1.0f)
-				.Padding(2.0f, 0.0f)
-				[
-					SNew(SEditableTextBox)
-					.HintText(LOCTEXT("ScreenshotFilterHint", "Filter Reports"))
-					.Text(FText::FromString(ReportFilterString))
-					.OnTextCommitted(this, &SScreenShotBrowser::OnFilterStringCommitted)
-					
-				]
-
-				+SHorizontalBox::Slot()
-				.AutoWidth()
-				.Padding(2.0f, 0.0f)
+				.HAlign(HAlign_Fill)
+				.Padding(10.0f, 0.0f)
 				[
 					SNew(SHorizontalBox)
 
 					+ SHorizontalBox::Slot()
-					.AutoWidth()
+					.FillWidth(1.0f)
+					.Padding(2.0f, 0.0f)
+					.HAlign(HAlign_Fill)
 					[
-						SNew(SCheckBox)
-						.HAlign(HAlign_Center)
-					.IsChecked(bDisplayingSuccess ? ECheckBoxState::Checked : ECheckBoxState::Unchecked)
-					.OnCheckStateChanged(this, &SScreenShotBrowser::DisplaySuccess_OnCheckStateChanged)
+						SNew(SSearchBox)
+						.HintText(LOCTEXT("ScreenshotFilterHint", "Search"))
+						.ToolTipText(LOCTEXT("Search Tests", "Search Tests"))
+						.OnTextCommitted(this, &SScreenShotBrowser::OnFilterStringCommitted)
 					]
 
-				+ SHorizontalBox::Slot()
+					+SHorizontalBox::Slot()
 					.AutoWidth()
-					.VAlign(VAlign_Center)
+					.Padding(2.0f, 0.0f)
 					[
-						SNew(STextBlock)
-						.AutoWrapText(true)
-					.Text(LOCTEXT("DisplaySuccess", "Show Passed"))
-					]
-				]
+						SNew(SHorizontalBox)
 
-				+ SHorizontalBox::Slot()
-				.AutoWidth()
-				.HAlign(HAlign_Right)
-				.Padding(2.0f, 0.0f)
-				[
-					SNew(SHorizontalBox)
-					
-					+ SHorizontalBox::Slot()
-					.AutoWidth()
-					[
-						SNew(SCheckBox)
-						.HAlign(HAlign_Center)
-						.IsChecked( bDisplayingError ? ECheckBoxState::Checked : ECheckBoxState::Unchecked)
-						.OnCheckStateChanged(this, &SScreenShotBrowser::DisplayError_OnCheckStateChanged)
-					]
+						+ SHorizontalBox::Slot()
+						.AutoWidth()
+						[
+							SNew(SCheckBox)
+							.HAlign(HAlign_Center)
+							.IsChecked(bDisplayingSuccess ? ECheckBoxState::Checked : ECheckBoxState::Unchecked)
+							.OnCheckStateChanged(this, &SScreenShotBrowser::DisplaySuccess_OnCheckStateChanged)
+						]
 
 					+ SHorizontalBox::Slot()
-					.AutoWidth()
-					.VAlign(VAlign_Center)
-					[
-						SNew(STextBlock)
-						.AutoWrapText(true)
-						.Text(LOCTEXT("DisplayErrors", "Show Fails"))
+						.AutoWidth()
+						.VAlign(VAlign_Center)
+						[
+							SNew(STextBlock)
+							.AutoWrapText(true)
+						.Text(LOCTEXT("DisplaySuccess", "Show Passed"))
+						]
 					]
-				]
 
-				+ SHorizontalBox::Slot()
-				.AutoWidth()
-				.Padding(2.0f, 0.0f)
-				[
-					SNew(SHorizontalBox)
-					
 					+ SHorizontalBox::Slot()
 					.AutoWidth()
+					.HAlign(HAlign_Right)
+					.Padding(2.0f, 0.0f)
 					[
-						SNew(SCheckBox)
-						.HAlign(HAlign_Center)
-						.IsChecked(bDisplayingNew ? ECheckBoxState::Checked : ECheckBoxState::Unchecked)
-						.OnCheckStateChanged(this, &SScreenShotBrowser::DisplayNew_OnCheckStateChanged)
+						SNew(SHorizontalBox)
+					
+						+ SHorizontalBox::Slot()
+						.AutoWidth()
+						[
+							SNew(SCheckBox)
+							.HAlign(HAlign_Center)
+							.IsChecked( bDisplayingError ? ECheckBoxState::Checked : ECheckBoxState::Unchecked)
+							.OnCheckStateChanged(this, &SScreenShotBrowser::DisplayError_OnCheckStateChanged)
+						]
+
+						+ SHorizontalBox::Slot()
+						.AutoWidth()
+						.VAlign(VAlign_Center)
+						[
+							SNew(STextBlock)
+							.AutoWrapText(true)
+							.Text(LOCTEXT("DisplayErrors", "Show Fails"))
+						]
 					]
 
 					+ SHorizontalBox::Slot()
 					.AutoWidth()
-					.VAlign(VAlign_Center)
+					.Padding(2.0f, 0.0f)
 					[
-						SNew(STextBlock)
-						.AutoWrapText(true)
-						.Text(LOCTEXT("DisplayNew", "Show New"))
-					]
-				]
-			]			
-		]
-
-		+ SVerticalBox::Slot()
-			.AutoHeight()
-		[
-			SNew(SHorizontalBox)
-
-			+ SHorizontalBox::Slot()
-			.FillWidth(1.0f)
-			[
-				SNew(SHorizontalBox)
-
-				+ SHorizontalBox::Slot()
-				.AutoWidth()
-				.Padding(2.f, 0.f)
-				[
-					SNew(SPositiveActionButton)
-					.Text(LOCTEXT("AddAllNewReports", "Add All New Reports"))
-					.ToolTipText(LOCTEXT("AddAllNewReportsTooltip", "Adds all new screenshots contained in the reports."))
-					.IsEnabled_Lambda([this]() -> bool
-						{
-							return ComparisonList.Num() > 0 && ISourceControlModule::Get().IsEnabled();
-						})
-					.OnClicked_Lambda([this]()
-						{
-							for (int Entry = ComparisonList.Num() - 1; Entry >= 0; --Entry)
-							{
-								TSharedPtr<FScreenComparisonModel> Model = ComparisonList[Entry];
-								const FImageComparisonResult& Comparison = Model->Report.GetComparisonResult();
-
-								if (CanAddNewReportResult(Comparison))
-								{
-									ComparisonList.RemoveAt(Entry);
-									Model->AddNew();
-
-									// Avoid thrashing P4
-									const float SleepBetweenOperations = 0.005f;
-									FPlatformProcess::Sleep(SleepBetweenOperations);
-								}
-							}
-
-							return FReply::Handled();
-						})
-				]
-
-				+ SHorizontalBox::Slot()
-				.AutoWidth()
-				.Padding(2.f, 0.f)
-				[
-					SNew(SNegativeActionButton)
-					.ActionButtonStyle(EActionButtonStyle::Warning)
-					.Text(LOCTEXT("ReplaceAllReports", "Replace All Reports"))
-					.ToolTipText(LOCTEXT("ReplaceAllReportsTooltip", "Replaces all screenshots containing a different result in the reports."))
-					.IsEnabled_Lambda([this]() -> bool
-						{
-							return ComparisonList.Num() > 0 && ISourceControlModule::Get().IsEnabled();
-						})
-					.OnClicked_Lambda([this]()
-						{
-							for (int Entry = ComparisonList.Num() - 1; Entry >= 0; --Entry)
-							{
-								TSharedPtr<FScreenComparisonModel> Model = ComparisonList[Entry];
-								const FImageComparisonResult& Comparison = Model->Report.GetComparisonResult();
-
-								if (!CanAddNewReportResult(Comparison))
-								{
-									ComparisonList.RemoveAt(Entry);
-									Model->Replace();
-
-									// Avoid thrashing P4
-									const float SleepBetweenOperations = 0.005f;
-									FPlatformProcess::Sleep(SleepBetweenOperations);
-								}
-							}
-
-							return FReply::Handled();
-						})
+						SNew(SHorizontalBox)
+					
+						+ SHorizontalBox::Slot()
+						.AutoWidth()
+						[
+							SNew(SCheckBox)
+							.HAlign(HAlign_Center)
+							.IsChecked(bDisplayingNew ? ECheckBoxState::Checked : ECheckBoxState::Unchecked)
+							.OnCheckStateChanged(this, &SScreenShotBrowser::DisplayNew_OnCheckStateChanged)
+						]
+
+						+ SHorizontalBox::Slot()
+						.AutoWidth()
+						.VAlign(VAlign_Center)
+						[
+							SNew(STextBlock)
+							.AutoWrapText(true)
+							.Text(LOCTEXT("DisplayNew", "Show New"))
+						]
 					]
-
-				+ SHorizontalBox::Slot()
-				.AutoWidth()
-				.Padding(2.f, 0.f)
-				[
-					SNew(SNegativeActionButton)
-					.Text(LOCTEXT("DeleteAllReports", "Delete All Reports"))
-					.ToolTipText(LOCTEXT("DeleteAllReportsTooltip", "Deletes all the current reports.  Reports are not removed unless the user resolves them, \nso if you just want to reset the state of the reports, clear them here and then re-run the tests."))
-					.IsEnabled_Lambda([this]() -> bool
-						{
-							return ComparisonList.Num() > 0;
-						})
-					.OnClicked_Lambda([this]()
-						{
-							while (ComparisonList.Num() > 0)
-							{
-								TSharedPtr<FScreenComparisonModel> Model = ComparisonList.Pop();
-								Model->Complete(true);
-							}
-
-							return FReply::Handled();
-						})
-				]
+					
+					// Empty horizontal slot that works as a spacer
+					+ SHorizontalBox::Slot()
+					.FillWidth(1.0f)
+					.Padding(2.0f, 0.0f)
+					.HAlign(HAlign_Fill)
+				]			
 			]
 		]
 		
