Add ability to filter out tests while typing text in the Search Bar of Screen Comparison tab.

diff --git a/Engine/Source/Developer/ScreenShotComparison/Private/Widgets/SScreenComparisonRow.cpp b/Engine/Source/Developer/ScreenShotComparison/Private/Widgets/SScreenComparisonRow.cpp
--- a/Engine/Source/Developer/ScreenShotComparison/Private/Widgets/SScreenComparisonRow.cpp
+++ b/Engine/Source/Developer/ScreenShotComparison/Private/Widgets/SScreenComparisonRow.cpp
@@ -176,13 +176,6 @@ TSharedRef<SWidget> SScreenComparisonRow::GenerateWidgetForColumn(const FName& C
 
 			FSlateColor TextColor = FSlateColor::UseForeground();
 
-			FString Name = FString::Printf(TEXT("%s.%s"), *ModelMetaData->Context, *ModelMetaData->ScreenShotName);
-			if ((ModelMetaData->Context.Len() && ModelMetaData->TestName.Len())
-				|| !ModelMetaData->ScreenShotName.Len())
-			{
-				Name = FString::Printf(TEXT("%s.%s"), *ModelMetaData->Context, *ModelMetaData->TestName);
-			}
-
 			if (ComparisonResult.IsNew())
 			{
 				TextColor = FSlateColor(FLinearColor::Yellow);
@@ -193,12 +186,12 @@ TSharedRef<SWidget> SScreenComparisonRow::GenerateWidgetForColumn(const FName& C
 			}
 
 			return SNew(STextBlock)
-				.Text(FText::FromString(Name))
+				.Text(GetName())
 				.ColorAndOpacity(TextColor);
 		}
 		else
 		{
-			return SNew(STextBlock).Text(LOCTEXT("Unknown", "Unknown Test, no metadata discovered."));
+			return SNew(STextBlock).Text(GetName());
 		}
 	}
 	else if (ColumnName == "Date")
@@ -655,4 +648,29 @@ FReply SScreenComparisonRow::OnImageClicked(const FGeometry& InGeometry, const F
 	return FReply::Handled();
 }
 
+FText SScreenComparisonRow::GetName() const
+{
+	if (!Name.IsSet())
+	{
+		auto ModelMetaData = Model->GetMetadata();
+		if (ModelMetaData.IsSet())
+		{
+			FString NameString = FString::Printf(TEXT("%s.%s"), *ModelMetaData->Context, *ModelMetaData->ScreenShotName);
+			if ((ModelMetaData->Context.Len() && ModelMetaData->TestName.Len())
+				|| !ModelMetaData->ScreenShotName.Len())
+			{
+				NameString = FString::Printf(TEXT("%s.%s"), *ModelMetaData->Context, *ModelMetaData->TestName);
+			}
+
+			Name = FText::FromString(NameString);
+		}
+		else
+		{
+			Name = LOCTEXT("Unknown", "Unknown Test, no metadata discovered.");
+		}
+	}
+
+	return Name.GetValue();
+}
+
 #undef LOCTEXT_NAMESPACE
diff --git a/Engine/Source/Developer/ScreenShotComparison/Private/Widgets/SScreenComparisonRow.h b/Engine/Source/Developer/ScreenShotComparison/Private/Widgets/SScreenComparisonRow.h
--- a/Engine/Source/Developer/ScreenShotComparison/Private/Widgets/SScreenComparisonRow.h
+++ b/Engine/Source/Developer/ScreenShotComparison/Private/Widgets/SScreenComparisonRow.h
@@ -42,6 +42,14 @@ public:
 	
 	virtual TSharedRef<SWidget> GenerateWidgetForColumn(const FName& ColumnName) override;
 
+	
+	/**
+	 * Get content of the Name column.
+	 *
+	 * @return FText object corresponding to the current object's content of the Name column.
+	 */
+	FText GetName() const;
+
 private:
 	bool CanUseSourceControl() const;
 	bool IsComparingAgainstPlatformFallback() const;
@@ -73,6 +81,8 @@ private:
 	//Holds the screen shot info.
 	TSharedPtr<FScreenComparisonModel> Model;
 
+	mutable TOptional<FText> Name;
+
 	// The manager containing the screen shots
 	IScreenShotManagerPtr ScreenshotManager;
 
diff --git a/Engine/Source/Developer/ScreenShotComparison/Private/Widgets/SScreenShotBrowser.cpp b/Engine/Source/Developer/ScreenShotComparison/Private/Widgets/SScreenShotBrowser.cpp
--- a/Engine/Source/Developer/ScreenShotComparison/Private/Widgets/SScreenShotBrowser.cpp
+++ b/Engine/Source/Developer/ScreenShotComparison/Private/Widgets/SScreenShotBrowser.cpp
@@ -184,7 +184,7 @@ void SScreenShotBrowser::Construct( const FArguments& InArgs,  IScreenShotManage
 						SNew(SSearchBox)
 						.HintText(LOCTEXT("ScreenshotFilterHint", "Search"))
 						.ToolTipText(LOCTEXT("Search Tests", "Search Tests"))
-						.OnTextCommitted(this, &SScreenShotBrowser::OnFilterStringCommitted)
+						.OnTextChanged(this, &SScreenShotBrowser::OnReportFilterTextChanged)
 					]
 
 					+SHorizontalBox::Slot()
@@ -371,12 +371,17 @@ void SScreenShotBrowser::OnReportsChanged(const TArray<struct FFileChangeData>&
 
 TSharedRef<ITableRow> SScreenShotBrowser::OnGenerateWidgetForScreenResults(TSharedPtr<FScreenComparisonModel> InItem, const TSharedRef<STableViewBase>& OwnerTable)
 {
-	// Create the row widget.
-	return
-		SNew(SScreenComparisonRow, OwnerTable)
+	check(InItem.IsValid());
+
+	TSharedRef<SScreenComparisonRow> ResultWidget = SNew(SScreenComparisonRow, OwnerTable)
 		.ScreenshotManager(ScreenShotManager)
 		.ComparisonDirectory(ComparisonRoot)
 		.ComparisonResult(InItem);
+
+	const bool bVisible = (ReportFilterString.IsEmpty() || ResultWidget->GetName().ToString().Contains(ReportFilterString, ESearchCase::IgnoreCase));
+	ResultWidget->SetVisibility(bVisible ? EVisibility::Visible : EVisibility::Collapsed);
+
+	return ResultWidget;
 }
 
 void SScreenShotBrowser::DisplaySuccess_OnCheckStateChanged(ECheckBoxState NewRadioState)
@@ -397,16 +402,42 @@ void SScreenShotBrowser::DisplayNew_OnCheckStateChanged(ECheckBoxState NewRadioS
 	bReportsChanged = true;
 }
 
-void SScreenShotBrowser::OnFilterStringCommitted(const FText& InText, ETextCommit::Type InCommitType)
+void SScreenShotBrowser::OnReportFilterTextChanged(const FText& InText)
 {
-	FString InString = InText.ToString();
-	if (ReportFilterString.Compare(InString, ESearchCase::IgnoreCase) != 0)
+	ReportFilterString = InText.ToString();
+	if (ApplyReportFilterToVWidgets() > 0)
 	{
-		ReportFilterString = InString;
-		bReportsChanged = true;
+		ComparisonView->RequestListRefresh();
 	}
 }
 
+uint32 SScreenShotBrowser::ApplyReportFilterToVWidgets()
+{
+	uint32 TouchedWidgetsCount = 0;
+	for (auto Item : ComparisonList)
+	{
+		TSharedPtr<ITableRow> Widget = ComparisonView->WidgetFromItem(Item);
+		// Note that some widgets that are not visible in the GUI view yet will be constructed only after these widgets are about to appear in the view.
+		// So, we do not use check(Widget.IsValid()) here. Visibility for these items will be set during generation of the widgets.
+		if (Widget.IsValid())
+		{
+			const EVisibility CurrentVisibility = Widget->AsWidget()->GetVisibility();
+
+			SScreenComparisonRow* Row = static_cast<SScreenComparisonRow*>(Widget.Get());
+			const bool bVisible = (ReportFilterString.IsEmpty() || Row->GetName().ToString().Contains(ReportFilterString, ESearchCase::IgnoreCase));
+			const EVisibility DesiredVisibility = (bVisible ? EVisibility::Visible : EVisibility::Collapsed);
+
+			if (DesiredVisibility != CurrentVisibility)
+			{
+				++TouchedWidgetsCount;
+				Widget->AsWidget()->SetVisibility(DesiredVisibility);
+			}
+		}
+	}
+
+	return TouchedWidgetsCount;
+}
+
 void SScreenShotBrowser::RebuildTree()
 {
 	bReportsChanged = false;
@@ -440,11 +471,6 @@ void SScreenShotBrowser::RebuildTree()
 				continue;
 			}
 
-			if (ReportFilterString.Len() && !Report.GetReportPath().Contains(ReportFilterString, ESearchCase::IgnoreCase))
-			{
-				continue;
-			}
-
 			TSharedPtr<FScreenComparisonModel> Model = MakeShared<FScreenComparisonModel>(Report);
 			Model->OnComplete.AddLambda([this, Model] () {
 				ComparisonList.Remove(Model);
diff --git a/Engine/Source/Developer/ScreenShotComparison/Private/Widgets/SScreenShotBrowser.h b/Engine/Source/Developer/ScreenShotComparison/Private/Widgets/SScreenShotBrowser.h
--- a/Engine/Source/Developer/ScreenShotComparison/Private/Widgets/SScreenShotBrowser.h
+++ b/Engine/Source/Developer/ScreenShotComparison/Private/Widgets/SScreenShotBrowser.h
@@ -46,7 +46,7 @@ public:
 	void DisplaySuccess_OnCheckStateChanged(ECheckBoxState NewRadioState);
 	void DisplayError_OnCheckStateChanged(ECheckBoxState NewRadioState);
 	void DisplayNew_OnCheckStateChanged(ECheckBoxState NewRadioState);
-	void OnFilterStringCommitted(const FText& InText, ETextCommit::Type InCommitType);
+	void OnReportFilterTextChanged(const FText& InText);
 
 private:
 
@@ -63,6 +63,13 @@ private:
 
 	bool CanAddNewReportResult(const FImageComparisonResult& Comparison);
 
+	/**
+	* Apply the current report filter string to the widgets of comparison view.
+	* 
+	* @return amount of the widgets that have changed visibility property in accordance with the current report filter string.
+	*/
+	uint32 ApplyReportFilterToVWidgets();
+
 private:
 
 	// The manager containing the screen shots
