Cleaned up SScreenShotBrowser to no longer require rebuilding of the tree whenever the display criteria is changed to increase performance of displaying the impacted image comparisons

#rnx

diff --git a/Engine/Source/Developer/ScreenShotComparison/Private/Widgets/SScreenShotBrowser.cpp b/Engine/Source/Developer/ScreenShotComparison/Private/Widgets/SScreenShotBrowser.cpp
--- a/Engine/Source/Developer/ScreenShotComparison/Private/Widgets/SScreenShotBrowser.cpp
+++ b/Engine/Source/Developer/ScreenShotComparison/Private/Widgets/SScreenShotBrowser.cpp
@@ -378,7 +378,7 @@ TSharedRef<ITableRow> SScreenShotBrowser::OnGenerateWidgetForScreenResults(TShar
 		.ComparisonDirectory(ComparisonRoot)
 		.ComparisonResult(InItem);
 
-	const bool bVisible = (ReportFilterString.IsEmpty() || ResultWidget->GetName().ToString().Contains(ReportFilterString, ESearchCase::IgnoreCase));
+	const bool bVisible = MatchesReportFilterCriteria(ResultWidget->GetName().ToString(), InItem->Report.GetComparisonResult());
 	ResultWidget->SetVisibility(bVisible ? EVisibility::Visible : EVisibility::Collapsed);
 
 	return ResultWidget;
@@ -387,31 +387,55 @@ TSharedRef<ITableRow> SScreenShotBrowser::OnGenerateWidgetForScreenResults(TShar
 void SScreenShotBrowser::DisplaySuccess_OnCheckStateChanged(ECheckBoxState NewRadioState)
 {
 	bDisplayingSuccess = (NewRadioState == ECheckBoxState::Checked);
-	bReportsChanged = true;
+	ApplyReportFilterToVWidgets();
 }
 
 void SScreenShotBrowser::DisplayError_OnCheckStateChanged(ECheckBoxState NewRadioState)
 {
 	bDisplayingError = (NewRadioState == ECheckBoxState::Checked);
-	bReportsChanged = true;
+	ApplyReportFilterToVWidgets();
 }
 
 void SScreenShotBrowser::DisplayNew_OnCheckStateChanged(ECheckBoxState NewRadioState)
 {
 	bDisplayingNew = (NewRadioState == ECheckBoxState::Checked);
-	bReportsChanged = true;
+	ApplyReportFilterToVWidgets();
 }
 
 void SScreenShotBrowser::OnReportFilterTextChanged(const FText& InText)
 {
 	ReportFilterString = InText.ToString();
-	if (ApplyReportFilterToVWidgets() > 0)
+	ApplyReportFilterToVWidgets();
+}
+
+bool SScreenShotBrowser::MatchesReportFilterCriteria(const FString& name, const FImageComparisonResult& ComparisonResult) const
+{
+	if (!ReportFilterString.IsEmpty() && !name.Contains(ReportFilterString, ESearchCase::IgnoreCase))
 	{
-		ComparisonView->RequestListRefresh();
+		return false;
+	}
+
+	bool bIsNew = ComparisonResult.IsNew();
+	bool bIsFail = !ComparisonResult.AreSimilar();
+	bool bIsPass = !bIsNew && !bIsFail;
+
+	if (bIsPass && !bDisplayingSuccess)
+	{
+		return false;
+	}
+	if (bIsNew && !bDisplayingNew)
+	{
+		return false;
+	}
+	if (bIsFail && !bDisplayingError)
+	{
+		return false;
 	}
+
+	return true;
 }
 
-uint32 SScreenShotBrowser::ApplyReportFilterToVWidgets()
+void SScreenShotBrowser::ApplyReportFilterToVWidgets()
 {
 	uint32 TouchedWidgetsCount = 0;
 	for (auto Item : ComparisonList)
@@ -424,7 +448,7 @@ uint32 SScreenShotBrowser::ApplyReportFilterToVWidgets()
 			const EVisibility CurrentVisibility = Widget->AsWidget()->GetVisibility();
 
 			SScreenComparisonRow* Row = static_cast<SScreenComparisonRow*>(Widget.Get());
-			const bool bVisible = (ReportFilterString.IsEmpty() || Row->GetName().ToString().Contains(ReportFilterString, ESearchCase::IgnoreCase));
+			const bool bVisible = MatchesReportFilterCriteria(Row->GetName().ToString(), Item->Report.GetComparisonResult());
 			const EVisibility DesiredVisibility = (bVisible ? EVisibility::Visible : EVisibility::Collapsed);
 
 			if (DesiredVisibility != CurrentVisibility)
@@ -435,7 +459,10 @@ uint32 SScreenShotBrowser::ApplyReportFilterToVWidgets()
 		}
 	}
 
-	return TouchedWidgetsCount;
+	if (TouchedWidgetsCount > 0)
+	{
+		ComparisonView->RequestListRefresh();
+	}
 }
 
 void SScreenShotBrowser::RebuildTree()
@@ -450,27 +477,6 @@ void SScreenShotBrowser::RebuildTree()
 
 		for ( const FComparisonReport& Report : CurrentReports )
 		{
-			const FImageComparisonResult& Comparison = Report.GetComparisonResult();
-
-			bool IsNew = Comparison.IsNew();
-			bool IsFail = !Comparison.AreSimilar();
-			bool IsPass = !IsNew && !IsFail;
-
-			if (IsPass && !bDisplayingSuccess)
-			{
-				continue;
-			}
-
-			if (IsNew && !bDisplayingNew)
-			{
-				continue;
-			}
-
-			if (IsFail && !bDisplayingError)
-			{
-				continue;
-			}
-
 			TSharedPtr<FScreenComparisonModel> Model = MakeShared<FScreenComparisonModel>(Report);
 			Model->OnComplete.AddLambda([this, Model] () {
 				ComparisonList.Remove(Model);
diff --git a/Engine/Source/Developer/ScreenShotComparison/Private/Widgets/SScreenShotBrowser.h b/Engine/Source/Developer/ScreenShotComparison/Private/Widgets/SScreenShotBrowser.h
--- a/Engine/Source/Developer/ScreenShotComparison/Private/Widgets/SScreenShotBrowser.h
+++ b/Engine/Source/Developer/ScreenShotComparison/Private/Widgets/SScreenShotBrowser.h
@@ -64,11 +64,18 @@ private:
 	bool CanAddNewReportResult(const FImageComparisonResult& Comparison);
 
 	/**
-	* Apply the current report filter string to the widgets of comparison view.
-	* 
-	* @return amount of the widgets that have changed visibility property in accordance with the current report filter string.
+	* Checks whether or not the information is valid against the report filtering criteria.
+	*
+	* @param ItemName - The name of the item for the string filtering to be checked against.
+	* @param ComparisonResult - The ImageComparisonResult used to be checked against the display criteria.
+	* @return true if the passed in information matches the current filtering criteria
 	*/
-	uint32 ApplyReportFilterToVWidgets();
+	bool MatchesReportFilterCriteria(const FString& ItemName, const FImageComparisonResult& ComparisonResult) const;
+
+	/**
+	* Apply the current report filter to the widgets of comparison view.
+	*/
+	void ApplyReportFilterToVWidgets();
 
 private:
 
