Fix issue with Screenshot comparison scrollbar by using a second array for filtered items

#rnx

diff --git a/Engine/Source/Developer/ScreenShotComparison/Private/Models/ScreenComparisonModel.cpp b/Engine/Source/Developer/ScreenShotComparison/Private/Models/ScreenComparisonModel.cpp
--- a/Engine/Source/Developer/ScreenShotComparison/Private/Models/ScreenComparisonModel.cpp
+++ b/Engine/Source/Developer/ScreenShotComparison/Private/Models/ScreenComparisonModel.cpp
@@ -73,6 +73,27 @@ TOptional<FAutomationScreenshotMetadata> FScreenComparisonModel::GetMetadata()
 	return Metadata;
 }
 
+FString FScreenComparisonModel::GetName()
+{
+	if (Name.IsEmpty())
+	{
+		auto LoadedMetadata = GetMetadata();
+		if (LoadedMetadata.IsSet())
+		{
+			FString NameString = FString::Printf(TEXT("%s.%s"), *LoadedMetadata->Context, *LoadedMetadata->ScreenShotName);
+			if ((LoadedMetadata->Context.Len() && LoadedMetadata->TestName.Len())
+				|| !LoadedMetadata->ScreenShotName.Len())
+			{
+				NameString = FString::Printf(TEXT("%s.%s"), *LoadedMetadata->Context, *LoadedMetadata->TestName);
+			}
+
+			Name = NameString;
+		}
+	}
+
+	return Name;
+}
+
 bool FScreenComparisonModel::AddNew()
 {
 	bool bSuccess = true;
diff --git a/Engine/Source/Developer/ScreenShotComparison/Private/Models/ScreenComparisonModel.h b/Engine/Source/Developer/ScreenShotComparison/Private/Models/ScreenComparisonModel.h
--- a/Engine/Source/Developer/ScreenShotComparison/Private/Models/ScreenComparisonModel.h
+++ b/Engine/Source/Developer/ScreenShotComparison/Private/Models/ScreenComparisonModel.h
@@ -27,6 +27,7 @@ public:
 	bool AddAlternative();
 
 	TOptional<FAutomationScreenshotMetadata> GetMetadata();
+	FString GetName();
 
 private:
 
@@ -36,6 +37,7 @@ private:
 	bool bComplete;
 
 	TOptional<FAutomationScreenshotMetadata> Metadata;
+	FString Name;
 
 	struct FFileMapping
 	{
diff --git a/Engine/Source/Developer/ScreenShotComparison/Private/Widgets/SScreenComparisonRow.cpp b/Engine/Source/Developer/ScreenShotComparison/Private/Widgets/SScreenComparisonRow.cpp
--- a/Engine/Source/Developer/ScreenShotComparison/Private/Widgets/SScreenComparisonRow.cpp
+++ b/Engine/Source/Developer/ScreenShotComparison/Private/Widgets/SScreenComparisonRow.cpp
@@ -652,17 +652,10 @@ FText SScreenComparisonRow::GetName() const
 {
 	if (!Name.IsSet())
 	{
-		auto ModelMetaData = Model->GetMetadata();
-		if (ModelMetaData.IsSet())
+		FString ModelName = Model->GetName();
+		if (!ModelName.IsEmpty())
 		{
-			FString NameString = FString::Printf(TEXT("%s.%s"), *ModelMetaData->Context, *ModelMetaData->ScreenShotName);
-			if ((ModelMetaData->Context.Len() && ModelMetaData->TestName.Len())
-				|| !ModelMetaData->ScreenShotName.Len())
-			{
-				NameString = FString::Printf(TEXT("%s.%s"), *ModelMetaData->Context, *ModelMetaData->TestName);
-			}
-
-			Name = FText::FromString(NameString);
+			Name = FText::FromString(ModelName);
 		}
 		else
 		{
diff --git a/Engine/Source/Developer/ScreenShotComparison/Private/Widgets/SScreenShotBrowser.cpp b/Engine/Source/Developer/ScreenShotComparison/Private/Widgets/SScreenShotBrowser.cpp
--- a/Engine/Source/Developer/ScreenShotComparison/Private/Widgets/SScreenShotBrowser.cpp
+++ b/Engine/Source/Developer/ScreenShotComparison/Private/Widgets/SScreenShotBrowser.cpp
@@ -65,21 +65,22 @@ void SScreenShotBrowser::Construct( const FArguments& InArgs,  IScreenShotManage
 			[
 				SNew(SPositiveActionButton)
 				.Text(LOCTEXT("AddAllNewReports", "Add All New Reports"))
-				.ToolTipText(LOCTEXT("AddAllNewReportsTooltip", "Adds all new screenshots contained in the reports."))
+				.ToolTipText(LOCTEXT("AddAllNewReportsTooltip", "Adds all filtered new screenshots contained in the reports."))
 				.IsEnabled_Lambda([this]() -> bool
 					{
-						return ComparisonList.Num() > 0 && ISourceControlModule::Get().IsEnabled();
+						return FilteredComparisonList.Num() > 0 && ISourceControlModule::Get().IsEnabled();
 					})
 				.OnClicked_Lambda([this]()
 					{
-						for (int Entry = ComparisonList.Num() - 1; Entry >= 0; --Entry)
+						for (int Entry = FilteredComparisonList.Num() - 1; Entry >= 0; --Entry)
 						{
-							TSharedPtr<FScreenComparisonModel> Model = ComparisonList[Entry];
+							TSharedPtr<FScreenComparisonModel> Model = FilteredComparisonList[Entry];
 							const FImageComparisonResult& Comparison = Model->Report.GetComparisonResult();
 
 							if (CanAddNewReportResult(Comparison))
 							{
-								ComparisonList.RemoveAt(Entry);
+								FilteredComparisonList.RemoveAt(Entry);
+								ComparisonList.Remove(Model);
 								Model->AddNew();
 
 								// Avoid thrashing P4
@@ -99,21 +100,22 @@ void SScreenShotBrowser::Construct( const FArguments& InArgs,  IScreenShotManage
 				SNew(SNegativeActionButton)
 				.ActionButtonStyle(EActionButtonStyle::Warning)
 				.Text(LOCTEXT("ReplaceAllReports", "Replace All Reports"))
-				.ToolTipText(LOCTEXT("ReplaceAllReportsTooltip", "Replaces all screenshots containing a different result in the reports."))
+				.ToolTipText(LOCTEXT("ReplaceAllReportsTooltip", "Replaces all filtered screenshots containing a different result in the reports."))
 				.IsEnabled_Lambda([this]() -> bool
 					{
-						return ComparisonList.Num() > 0 && ISourceControlModule::Get().IsEnabled();
+						return FilteredComparisonList.Num() > 0 && ISourceControlModule::Get().IsEnabled();
 					})
 				.OnClicked_Lambda([this]()
 					{
-						for (int Entry = ComparisonList.Num() - 1; Entry >= 0; --Entry)
+						for (int Entry = FilteredComparisonList.Num() - 1; Entry >= 0; --Entry)
 						{
-							TSharedPtr<FScreenComparisonModel> Model = ComparisonList[Entry];
+							TSharedPtr<FScreenComparisonModel> Model = FilteredComparisonList[Entry];
 							const FImageComparisonResult& Comparison = Model->Report.GetComparisonResult();
 
 							if (!CanAddNewReportResult(Comparison))
 							{
-								ComparisonList.RemoveAt(Entry);
+								FilteredComparisonList.RemoveAt(Entry);
+								ComparisonList.Remove(Model);
 								Model->Replace();
 
 								// Avoid thrashing P4
@@ -132,16 +134,16 @@ void SScreenShotBrowser::Construct( const FArguments& InArgs,  IScreenShotManage
 			[
 				SNew(SNegativeActionButton)
 				.Text(LOCTEXT("DeleteAllReports", "Delete All Reports"))
-				.ToolTipText(LOCTEXT("DeleteAllReportsTooltip", "Deletes all the current reports.  Reports are not removed unless the user resolves them, \nso if you just want to reset the state of the reports, clear them here and then re-run the tests."))
+				.ToolTipText(LOCTEXT("DeleteAllReportsTooltip", "Deletes all the filtered reports. Reports are not removed unless the user resolves them, \nso if you just want to reset the state of the reports, clear them here and then re-run the tests."))
 				.IsEnabled_Lambda([this]() -> bool
 					{
-						return ComparisonList.Num() > 0;
+						return FilteredComparisonList.Num() > 0;
 					})
 				.OnClicked_Lambda([this]()
 					{
-						while (ComparisonList.Num() > 0)
+						while (FilteredComparisonList.Num() > 0)
 						{
-							TSharedPtr<FScreenComparisonModel> Model = ComparisonList.Pop();
+							TSharedPtr<FScreenComparisonModel> Model = FilteredComparisonList.Pop();
 							Model->Complete(true);
 						}
 
@@ -276,7 +278,7 @@ void SScreenShotBrowser::Construct( const FArguments& InArgs,  IScreenShotManage
 		.FillHeight( 1.0f )
 		[
 			SAssignNew(ComparisonView, SListView< TSharedPtr<FScreenComparisonModel> >)
-			.ListItemsSource(&ComparisonList)
+			.ListItemsSource(&FilteredComparisonList)
 			.OnGenerateRow(this, &SScreenShotBrowser::OnGenerateWidgetForScreenResults)
 			.SelectionMode(ESelectionMode::None)
 			.HeaderRow
@@ -378,9 +380,6 @@ TSharedRef<ITableRow> SScreenShotBrowser::OnGenerateWidgetForScreenResults(TShar
 		.ComparisonDirectory(ComparisonRoot)
 		.ComparisonResult(InItem);
 
-	const bool bVisible = MatchesReportFilterCriteria(ResultWidget->GetName().ToString(), InItem->Report.GetComparisonResult());
-	ResultWidget->SetVisibility(bVisible ? EVisibility::Visible : EVisibility::Collapsed);
-
 	return ResultWidget;
 }
 
@@ -408,9 +407,9 @@ void SScreenShotBrowser::OnReportFilterTextChanged(const FText& InText)
 	ApplyReportFilterToVWidgets();
 }
 
-bool SScreenShotBrowser::MatchesReportFilterCriteria(const FString& name, const FImageComparisonResult& ComparisonResult) const
+bool SScreenShotBrowser::MatchesReportFilterCriteria(const FString& ComparisonName, const FImageComparisonResult& ComparisonResult) const
 {
-	if (!ReportFilterString.IsEmpty() && !name.Contains(ReportFilterString, ESearchCase::IgnoreCase))
+	if (!ReportFilterString.IsEmpty() && !ComparisonName.Contains(ReportFilterString, ESearchCase::IgnoreCase))
 	{
 		return false;
 	}
@@ -437,38 +436,22 @@ bool SScreenShotBrowser::MatchesReportFilterCriteria(const FString& name, const
 
 void SScreenShotBrowser::ApplyReportFilterToVWidgets()
 {
-	uint32 TouchedWidgetsCount = 0;
+	FilteredComparisonList.Reset();
 	for (auto Item : ComparisonList)
 	{
-		TSharedPtr<ITableRow> Widget = ComparisonView->WidgetFromItem(Item);
-		// Note that some widgets that are not visible in the GUI view yet will be constructed only after these widgets are about to appear in the view.
-		// So, we do not use check(Widget.IsValid()) here. Visibility for these items will be set during generation of the widgets.
-		if (Widget.IsValid())
+		if(Item.IsValid() && MatchesReportFilterCriteria(Item->GetName(), Item->Report.GetComparisonResult()))
 		{
-			const EVisibility CurrentVisibility = Widget->AsWidget()->GetVisibility();
-
-			SScreenComparisonRow* Row = static_cast<SScreenComparisonRow*>(Widget.Get());
-			const bool bVisible = MatchesReportFilterCriteria(Row->GetName().ToString(), Item->Report.GetComparisonResult());
-			const EVisibility DesiredVisibility = (bVisible ? EVisibility::Visible : EVisibility::Collapsed);
-
-			if (DesiredVisibility != CurrentVisibility)
-			{
-				++TouchedWidgetsCount;
-				Widget->AsWidget()->SetVisibility(DesiredVisibility);
-			}
+			FilteredComparisonList.Add(Item);
 		}
 	}
-
-	if (TouchedWidgetsCount > 0)
-	{
-		ComparisonView->RequestListRefresh();
-	}
+	ComparisonView->RequestListRefresh();
 }
 
 void SScreenShotBrowser::RebuildTree()
 {
 	bReportsChanged = false;
 	ComparisonList.Reset();
+	FilteredComparisonList.Reset();
 
 	if ( ScreenShotManager->OpenComparisonReports(ComparisonRoot, CurrentReports) )
 	{
@@ -480,6 +463,7 @@ void SScreenShotBrowser::RebuildTree()
 			TSharedPtr<FScreenComparisonModel> Model = MakeShared<FScreenComparisonModel>(Report);
 			Model->OnComplete.AddLambda([this, Model] () {
 				ComparisonList.Remove(Model);
+				FilteredComparisonList.Remove(Model);
 				ComparisonView->RequestListRefresh();
 			});
 
@@ -487,7 +471,7 @@ void SScreenShotBrowser::RebuildTree()
 		}
 	}
 
-	ComparisonView->RequestListRefresh();
+	ApplyReportFilterToVWidgets();
 }
 
 bool SScreenShotBrowser::CanAddNewReportResult(const FImageComparisonResult& Comparison)
diff --git a/Engine/Source/Developer/ScreenShotComparison/Private/Widgets/SScreenShotBrowser.h b/Engine/Source/Developer/ScreenShotComparison/Private/Widgets/SScreenShotBrowser.h
--- a/Engine/Source/Developer/ScreenShotComparison/Private/Widgets/SScreenShotBrowser.h
+++ b/Engine/Source/Developer/ScreenShotComparison/Private/Widgets/SScreenShotBrowser.h
@@ -66,11 +66,11 @@ private:
 	/**
 	* Checks whether or not the information is valid against the report filtering criteria.
 	*
-	* @param ItemName - The name of the item for the string filtering to be checked against.
+	* @param ComparisonName - The name of the item to be checked against.
 	* @param ComparisonResult - The ImageComparisonResult used to be checked against the display criteria.
 	* @return true if the passed in information matches the current filtering criteria
 	*/
-	bool MatchesReportFilterCriteria(const FString& ItemName, const FImageComparisonResult& ComparisonResult) const;
+	bool MatchesReportFilterCriteria(const FString& ComparisonName, const FImageComparisonResult& ComparisonResult) const;
 
 	/**
 	* Apply the current report filter to the widgets of comparison view.
@@ -88,9 +88,12 @@ private:
 	/** The imported screenshot results */
 	TArray<FComparisonReport> CurrentReports;
 
-	/** The imported screenshot results copied into an array usable by the list view */
+	/** The imported screenshot results copied into an array to be filtered */
 	TArray<TSharedPtr<FScreenComparisonModel>> ComparisonList;
 
+	/** The filtered screenshot results copied into an array usable by the list view */
+	TArray<TSharedPtr<FScreenComparisonModel>> FilteredComparisonList;
+
 	/**  */
 	TSharedPtr< SListView< TSharedPtr<FScreenComparisonModel> > > ComparisonView;
 
