SessionFrontend - Make FScreenShotManager::OpenComparisonReports and metadata loading async

diff --git a/Engine/Source/Developer/ScreenShotComparison/Private/Models/ScreenComparisonModel.cpp b/Engine/Source/Developer/ScreenShotComparison/Private/Models/ScreenComparisonModel.cpp
--- a/Engine/Source/Developer/ScreenShotComparison/Private/Models/ScreenComparisonModel.cpp
+++ b/Engine/Source/Developer/ScreenShotComparison/Private/Models/ScreenComparisonModel.cpp
@@ -8,8 +8,6 @@
 #include "ISourceControlProvider.h"
 #include "Misc/Paths.h"
 
-DEFINE_LOG_CATEGORY_STATIC(LogScreenshotComparison, Log, All);
-
 FScreenComparisonModel::FScreenComparisonModel(const FComparisonReport& InReport)
 	: Report(InReport)
 	, bComplete(false)
@@ -27,6 +25,8 @@ FScreenComparisonModel::FScreenComparisonModel(const FComparisonReport& InReport
 
 	FileImports.Add(FFileMapping(OutputImage, SourceImage));
 	FileImports.Add(FFileMapping(FPaths::ChangeExtension(OutputImage, TEXT("json")), FPaths::ChangeExtension(SourceImage, TEXT("json"))));
+
+	TryLoadMetadata();
 }
 
 bool FScreenComparisonModel::IsComplete() const
@@ -46,29 +46,34 @@ void FScreenComparisonModel::Complete(bool bWasSuccessful)
 	OnComplete.Broadcast();
 }
 
-TOptional<FAutomationScreenshotMetadata> FScreenComparisonModel::GetMetadata()
+void FScreenComparisonModel::TryLoadMetadata()
 {
-	// Load it.
-	if ( !Metadata.IsSet() )
+	if (!Metadata.IsSet())
 	{
 		const FImageComparisonResult& Comparison = Report.GetComparisonResult();
 
 		FString IncomingImage = Report.GetReportRootDirectory() / Comparison.ReportIncomingFilePath;
 		FString IncomingMetadata = FPaths::ChangeExtension(IncomingImage, TEXT("json"));
 
-		if ( !IncomingMetadata.IsEmpty() )
+		if (!IncomingMetadata.IsEmpty())
 		{
 			FString Json;
-			if ( FFileHelper::LoadFileToString(Json, *IncomingMetadata) )
+			if (FFileHelper::LoadFileToString(Json, *IncomingMetadata))
 			{
 				FAutomationScreenshotMetadata LoadedMetadata;
-				if ( FJsonObjectConverter::JsonObjectStringToUStruct(Json, &LoadedMetadata, 0, 0) )
+				if (FJsonObjectConverter::JsonObjectStringToUStruct(Json, &LoadedMetadata, 0, 0))
 				{
 					Metadata = LoadedMetadata;
 				}
 			}
 		}
 	}
+}
+
+TOptional<FAutomationScreenshotMetadata> FScreenComparisonModel::GetMetadata()
+{
+	// Load it just in case.
+	TryLoadMetadata();
 
 	return Metadata;
 }
diff --git a/Engine/Source/Developer/ScreenShotComparison/Private/Models/ScreenComparisonModel.h b/Engine/Source/Developer/ScreenShotComparison/Private/Models/ScreenComparisonModel.h
--- a/Engine/Source/Developer/ScreenShotComparison/Private/Models/ScreenComparisonModel.h
+++ b/Engine/Source/Developer/ScreenShotComparison/Private/Models/ScreenComparisonModel.h
@@ -32,6 +32,7 @@ public:
 private:
 
 	bool RemoveExistingApproved();
+	void TryLoadMetadata();
 
 private:
 	bool bComplete;
diff --git a/Engine/Source/Developer/ScreenShotComparison/Private/Widgets/SScreenShotBrowser.cpp b/Engine/Source/Developer/ScreenShotComparison/Private/Widgets/SScreenShotBrowser.cpp
--- a/Engine/Source/Developer/ScreenShotComparison/Private/Widgets/SScreenShotBrowser.cpp
+++ b/Engine/Source/Developer/ScreenShotComparison/Private/Widgets/SScreenShotBrowser.cpp
@@ -12,7 +12,6 @@
 #include "Widgets/SScreenComparisonRow.h"
 #include "Widgets/Input/SCheckBox.h"
 #include "Widgets/Input/SSearchBox.h"
-#include "Models/ScreenComparisonModel.h"
 #include "Misc/FeedbackContext.h"
 #include "Styling/AppStyle.h"
 #include "Modules/ModuleManager.h"
@@ -280,6 +279,7 @@ void SScreenShotBrowser::Construct( const FArguments& InArgs,  IScreenShotManage
 			SAssignNew(ComparisonView, SListView< TSharedPtr<FScreenComparisonModel> >)
 			.ListItemsSource(&FilteredComparisonList)
 			.OnGenerateRow(this, &SScreenShotBrowser::OnGenerateWidgetForScreenResults)
+			.Visibility(this, &SScreenShotBrowser::GetReportsVisibility)
 			.SelectionMode(ESelectionMode::None)
 			.HeaderRow
 			(
@@ -322,6 +322,15 @@ void SScreenShotBrowser::Construct( const FArguments& InArgs,  IScreenShotManage
 				.VAlignCell(VAlign_Center)
 			)
 		]
+
+		+ SVerticalBox::Slot()
+		.HAlign(HAlign_Center)
+		.VAlign(VAlign_Center)
+		[
+			SNew(SThrobber)
+			.Visibility(this, &SScreenShotBrowser::GetReportsUpdatingThrobberVisibility)
+		]
+
 	];
 
 	RefreshDirectoryWatcher();
@@ -341,8 +350,11 @@ void SScreenShotBrowser::Tick(const FGeometry& AllottedGeometry, const double In
 {
 	if ( bReportsChanged )
 	{
-		RebuildTree();
+		RequestRebuildTree();
 	}
+
+	ContinueRebuildTreeIfReady();
+	FinishRebuildTreeIfReady();
 }
 
 void SScreenShotBrowser::OnDirectoryChanged(const FString& Directory)
@@ -386,25 +398,25 @@ TSharedRef<ITableRow> SScreenShotBrowser::OnGenerateWidgetForScreenResults(TShar
 void SScreenShotBrowser::DisplaySuccess_OnCheckStateChanged(ECheckBoxState NewRadioState)
 {
 	bDisplayingSuccess = (NewRadioState == ECheckBoxState::Checked);
-	ApplyReportFilterToVWidgets();
+	ApplyReportFilterToWidgets();
 }
 
 void SScreenShotBrowser::DisplayError_OnCheckStateChanged(ECheckBoxState NewRadioState)
 {
 	bDisplayingError = (NewRadioState == ECheckBoxState::Checked);
-	ApplyReportFilterToVWidgets();
+	ApplyReportFilterToWidgets();
 }
 
 void SScreenShotBrowser::DisplayNew_OnCheckStateChanged(ECheckBoxState NewRadioState)
 {
 	bDisplayingNew = (NewRadioState == ECheckBoxState::Checked);
-	ApplyReportFilterToVWidgets();
+	ApplyReportFilterToWidgets();
 }
 
 void SScreenShotBrowser::OnReportFilterTextChanged(const FText& InText)
 {
 	ReportFilterString = InText.ToString();
-	ApplyReportFilterToVWidgets();
+	ApplyReportFilterToWidgets();
 }
 
 bool SScreenShotBrowser::MatchesReportFilterCriteria(const FString& ComparisonName, const FImageComparisonResult& ComparisonResult) const
@@ -434,10 +446,10 @@ bool SScreenShotBrowser::MatchesReportFilterCriteria(const FString& ComparisonNa
 	return true;
 }
 
-void SScreenShotBrowser::ApplyReportFilterToVWidgets()
+void SScreenShotBrowser::ApplyReportFilterToWidgets()
 {
 	FilteredComparisonList.Reset();
-	for (auto Item : ComparisonList)
+	for (auto& Item : ComparisonList)
 	{
 		if(Item.IsValid() && MatchesReportFilterCriteria(Item->GetName(), Item->Report.GetComparisonResult()))
 		{
@@ -447,31 +459,112 @@ void SScreenShotBrowser::ApplyReportFilterToVWidgets()
 	ComparisonView->RequestListRefresh();
 }
 
-void SScreenShotBrowser::RebuildTree()
+void SScreenShotBrowser::RequestRebuildTree()
 {
 	bReportsChanged = false;
-	ComparisonList.Reset();
+
+	PendingOpenComparisonReportsResult.Reset();
+	PendingScreenComparisonModelsLoadResult.Reset();
+	ComparisonList.Empty();
 	FilteredComparisonList.Reset();
+	CurrentReports.Reset();
 
-	if ( ScreenShotManager->OpenComparisonReports(ComparisonRoot, CurrentReports) )
+	PendingOpenComparisonReportsResult = ScreenShotManager->OpenComparisonReportsAsync(ComparisonRoot);
+}
+
+void SScreenShotBrowser::ContinueRebuildTreeIfReady()
+{
+	check(IsInGameThread());
+
+	if (PendingOpenComparisonReportsResult.IsValid()
+		&& PendingOpenComparisonReportsResult.IsReady())
 	{
+		CurrentReports = PendingOpenComparisonReportsResult.Get();
+		PendingOpenComparisonReportsResult.Reset();
+
+		if ((!CurrentReports.IsValid()) || CurrentReports->IsEmpty())
+		{
+			return;
+		}
 		//Sort by what will resolve down to the Screenshots' names
-		CurrentReports.Sort([](const FComparisonReport& LHS, const FComparisonReport& RHS) { return LHS.GetReportPath().Compare(RHS.GetReportPath()) < 0; });
+		CurrentReports->Sort([](const FComparisonReport& LHS, const FComparisonReport& RHS) { return LHS.GetReportPath().Compare(RHS.GetReportPath()) < 0; });
 
-		for ( const FComparisonReport& Report : CurrentReports )
+		PendingScreenComparisonModelsLoadResult = Async(EAsyncExecution::ThreadPool, [CurrentReportsWPtr = CurrentReports.ToWeakPtr()]() -> TArray<TSharedPtr<FScreenComparisonModel>>
 		{
-			TSharedPtr<FScreenComparisonModel> Model = MakeShared<FScreenComparisonModel>(Report);
-			Model->OnComplete.AddLambda([this, Model] () {
+			int32 CurrentReportsNum = 0;
+			{
+				TSharedPtr<TArray<FComparisonReport>> CurrentReportsSPtr = CurrentReportsWPtr.Pin();
+				if (CurrentReportsSPtr.IsValid())
+				{
+					CurrentReportsNum = CurrentReportsSPtr->Num();
+				}
+				else
+				{
+					// The current reports list is outdated. Cancel the job.
+					return {};
+				}
+			}
+
+
+			if (CurrentReportsNum <= 0)
+			{
+				return {};
+			}
+
+			TArray<TSharedPtr<FScreenComparisonModel>> Result;
+			Result.Reserve(CurrentReportsNum);
+
+			for (int32 CurrentReportIndex = 0; CurrentReportIndex < CurrentReportsNum; ++CurrentReportIndex)
+			{
+				TSharedPtr<TArray<FComparisonReport>> CurrentReportsSPtr = CurrentReportsWPtr.Pin();
+				if (CurrentReportsSPtr.IsValid())
+				{
+					const FComparisonReport& Report = (*CurrentReportsSPtr)[CurrentReportIndex];
+					Result.Add(MakeShared<FScreenComparisonModel>(Report));
+				}
+				else
+				{
+					// The current reports list is outdated. Cancel the job.
+					return {};
+				}
+			}
+
+			return Result;
+		});
+	}
+}
+
+void SScreenShotBrowser::FinishRebuildTreeIfReady()
+{
+	check(IsInGameThread());
+
+	if (PendingScreenComparisonModelsLoadResult.IsValid()
+		&& PendingScreenComparisonModelsLoadResult.IsReady())
+	{
+		ComparisonList = PendingScreenComparisonModelsLoadResult.Get();
+		PendingScreenComparisonModelsLoadResult.Reset();
+
+		for (auto& Model : ComparisonList)
+		{
+			Model->OnComplete.AddLambda([this, Model]() {
 				ComparisonList.Remove(Model);
 				FilteredComparisonList.Remove(Model);
 				ComparisonView->RequestListRefresh();
 			});
-
-			ComparisonList.Add(Model);
 		}
+
+		ApplyReportFilterToWidgets();
 	}
+}
 
-	ApplyReportFilterToVWidgets();
+EVisibility SScreenShotBrowser::GetReportsVisibility() const
+{
+	return (PendingOpenComparisonReportsResult.IsValid() || PendingScreenComparisonModelsLoadResult.IsValid() ? EVisibility::Collapsed : EVisibility::Visible);
+}
+
+EVisibility SScreenShotBrowser::GetReportsUpdatingThrobberVisibility() const
+{
+	return (GetReportsVisibility() == EVisibility::Visible ? EVisibility::Collapsed : EVisibility::Visible);
 }
 
 bool SScreenShotBrowser::CanAddNewReportResult(const FImageComparisonResult& Comparison)
diff --git a/Engine/Source/Developer/ScreenShotComparison/Private/Widgets/SScreenShotBrowser.h b/Engine/Source/Developer/ScreenShotComparison/Private/Widgets/SScreenShotBrowser.h
--- a/Engine/Source/Developer/ScreenShotComparison/Private/Widgets/SScreenShotBrowser.h
+++ b/Engine/Source/Developer/ScreenShotComparison/Private/Widgets/SScreenShotBrowser.h
@@ -12,8 +12,8 @@
 #include "Interfaces/IScreenShotManager.h"
 #include "Widgets/Views/STableViewBase.h"
 #include "Widgets/Views/STableRow.h"
-
-class FScreenComparisonModel;
+#include "Widgets/Images/SThrobber.h"
+#include "Models/ScreenComparisonModel.h"
 
 /**
  * Implements a Slate widget for browsing active game sessions.
@@ -57,9 +57,29 @@ private:
 	void OnReportsChanged(const TArray<struct FFileChangeData>& /*FileChanges*/);
 
 	/**
-	 * Regenerate the widgets when the filter changes
+	 * Requests regeneration of the widgets when it is needed.
 	 */
-	void RebuildTree();
+	void RequestRebuildTree();
+
+	/**
+	* Rebuilds screenshot tree if there are any pending comparison reports that are ready to be displayed.
+	*/
+	void ContinueRebuildTreeIfReady();
+
+	/**
+	* Finishes rebuilding of screenshot tree if all the report model metadata is loaded.
+	*/
+	void FinishRebuildTreeIfReady();
+
+	/**
+	* Gets reports tree widget visibility.
+	*/
+	EVisibility GetReportsVisibility() const;
+
+	/**
+	* Gets throbber visibility (report tree rebuild process is in progress).
+	*/
+	EVisibility GetReportsUpdatingThrobberVisibility() const;
 
 	bool CanAddNewReportResult(const FImageComparisonResult& Comparison);
 
@@ -75,7 +95,7 @@ private:
 	/**
 	* Apply the current report filter to the widgets of comparison view.
 	*/
-	void ApplyReportFilterToVWidgets();
+	void ApplyReportFilterToWidgets();
 
 private:
 
@@ -86,7 +106,7 @@ private:
 	FString ComparisonRoot;
 
 	/** The imported screenshot results */
-	TArray<FComparisonReport> CurrentReports;
+	TSharedPtr<TArray<FComparisonReport>> CurrentReports;
 
 	/** The imported screenshot results copied into an array to be filtered */
 	TArray<TSharedPtr<FScreenComparisonModel>> ComparisonList;
@@ -109,9 +129,15 @@ private:
 	/** Whether or not we're currently displaying tests with errors */
 	bool bDisplayingError;
 
-	/** Whether or not we're currently displaying warnings*/
+	/** Whether or not we're currently displaying warnings */
 	bool bDisplayingNew;
 
 	/** Filter string for our reports, if any */
 	FString ReportFilterString;
+
+	/** Stores comparison reports that are loaded asynchronously if ready */
+	TFuture<TSharedPtr<TArray<FComparisonReport>>> PendingOpenComparisonReportsResult;
+
+	/** Stores comparison reports' model objects that are loaded asynchronously if ready */
+	TFuture<TArray<TSharedPtr<FScreenComparisonModel>>> PendingScreenComparisonModelsLoadResult;
 };
diff --git a/Engine/Source/Developer/ScreenShotComparisonTools/Private/ScreenShotManager.cpp b/Engine/Source/Developer/ScreenShotComparisonTools/Private/ScreenShotManager.cpp
--- a/Engine/Source/Developer/ScreenShotComparisonTools/Private/ScreenShotManager.cpp
+++ b/Engine/Source/Developer/ScreenShotComparisonTools/Private/ScreenShotManager.cpp
@@ -548,48 +548,91 @@ FScreenshotExportResult FScreenShotManager::ExportScreenshotComparisonResult(FSt
 	return Results;
 }
 
-bool FScreenShotManager::OpenComparisonReports(FString ImportPath, TArray<FComparisonReport>& OutReports)
+TFuture<TSharedPtr<TArray<FComparisonReport>>> FScreenShotManager::OpenComparisonReportsAsync(const FString& ImportPath)
 {
-	OutReports.Reset();
-
-	FPaths::NormalizeDirectoryName(ImportPath);
-	ImportPath += TEXT("/");
+	FString PreprocessedImportPath(ImportPath);
+	FPaths::NormalizeDirectoryName(PreprocessedImportPath);
+	PreprocessedImportPath += TEXT("/");
 
 	TArray<FString> ComparisonReportPaths;
-	IFileManager::Get().FindFilesRecursive(ComparisonReportPaths, *ImportPath, TEXT("Report.json"), /*Files=*/true, /*Directories=*/false, /*bClearFileNames=*/ false);
+	IFileManager::Get().FindFilesRecursive(ComparisonReportPaths, *PreprocessedImportPath, TEXT("Report.json"), true, false, false);
+	
+	// Note that if current PendingComparisonReportPaths is valid,
+	// it will be reset in the line below (it will cancel the corresponding loading task)
+	PendingComparisonReportPaths = MakeShared<TArray<FString>>(MoveTemp(ComparisonReportPaths));
 
-	for ( const FString& ReportPath : ComparisonReportPaths )
+	return Async(EAsyncExecution::ThreadPool, [this, ImportPathToUse = MoveTemp(PreprocessedImportPath), ReportPathsWPtr = PendingComparisonReportPaths.ToWeakPtr()]() -> TSharedPtr<TArray<FComparisonReport>>
 	{
-		FString JsonString;
-		if ( FFileHelper::LoadFileToString(JsonString, *ReportPath) )
+		int32 ReportPathsNum = 0;
 		{
-			TSharedRef< TJsonReader<> > JsonReader = TJsonReaderFactory<>::Create(JsonString);
-
-			TSharedPtr<FJsonObject> JsonComparisonReport;
-			if ( !FJsonSerializer::Deserialize(JsonReader, JsonComparisonReport) )
+			TSharedPtr<TArray<FString>> ReportPathsSPtr = ReportPathsWPtr.Pin();
+			if (ReportPathsSPtr.IsValid())
 			{
-				return false;
+				ReportPathsNum = ReportPathsSPtr->Num();
 			}
+			else
+			{
+				// The reports paths list is outdated. Cancel the job with default return value.
+				return nullptr;
+			}
+		}
 
-			FImageComparisonResult ComparisonResult;
-			ComparisonResult.SetInvalid();
-			if ( FJsonObjectConverter::JsonObjectToUStruct(JsonComparisonReport.ToSharedRef(), &ComparisonResult, 0, 0) )
+		if (ReportPathsNum <= 0)
+		{
+			return nullptr;
+		}
+
+		TSharedPtr<TArray<FComparisonReport>> ComparisonReports = MakeShared<TArray<FComparisonReport>>();
+		ComparisonReports->Reserve(ReportPathsNum);
+
+		for (int32 ReportPathIndex = 0; ReportPathIndex < ReportPathsNum; ++ReportPathIndex)
+		{
+			FString ReportPath;
 			{
-				if (ComparisonResult.IsValid())
+				TSharedPtr<TArray<FString>> ReportPathsSPtr = ReportPathsWPtr.Pin();
+				if (ReportPathsSPtr.IsValid())
 				{
-					FComparisonReport Report(ImportPath, ReportPath);
-					Report.SetComparisonResult(ComparisonResult);
-					OutReports.Add(Report);
+					ReportPath = (*ReportPathsSPtr)[ReportPathIndex];
+					check(!ReportPath.IsEmpty());
 				}
 				else
 				{
-					UE_LOG(LogScreenShotManager, Error, TEXT("Report %s has invalid version '%d' (Current Version=%d)"), *ReportPath, ComparisonResult.Version, int32(ComparisonResult.CurrentVersion));
+					// The reports paths list is outdated. Cancel the job with default return value.
+					return nullptr;
+				}
+			}
+
+			FString JsonString;
+			if (FFileHelper::LoadFileToString(JsonString, *ReportPath))
+			{
+				TSharedRef<TJsonReader<>> JsonReader = TJsonReaderFactory<>::Create(JsonString);
+
+				TSharedPtr<FJsonObject> JsonComparisonReport;
+				if (!FJsonSerializer::Deserialize(JsonReader, JsonComparisonReport))
+				{
+					return nullptr;
+				}
+
+				FImageComparisonResult ComparisonResult;
+				ComparisonResult.SetInvalid();
+				if (FJsonObjectConverter::JsonObjectToUStruct(JsonComparisonReport.ToSharedRef(), &ComparisonResult, 0, 0))
+				{
+					if (ComparisonResult.IsValid())
+					{
+						FComparisonReport Report = FComparisonReport(ImportPathToUse, ReportPath);
+						Report.SetComparisonResult(ComparisonResult);
+						ComparisonReports->Add(Report);
+					}
+					else
+					{
+						UE_LOG(LogScreenShotManager, Error, TEXT("Report %s has invalid version '%d' (Current Version=%d)"), *ReportPath, ComparisonResult.Version, int32(ComparisonResult.CurrentVersion));
+					}
 				}
 			}
 		}
-	}
 
-	return true;
+		return ComparisonReports;
+	});
 }
 
 TSharedPtr<FImageComparisonResult> FScreenShotManager::CompareImageSequence(const TMap<FString, FString>& Sequence, const FAutomationScreenshotMetadata& Metadata)
diff --git a/Engine/Source/Developer/ScreenShotComparisonTools/Private/ScreenShotManager.h b/Engine/Source/Developer/ScreenShotComparisonTools/Private/ScreenShotManager.h
--- a/Engine/Source/Developer/ScreenShotComparisonTools/Private/ScreenShotManager.h
+++ b/Engine/Source/Developer/ScreenShotComparisonTools/Private/ScreenShotManager.h
@@ -32,7 +32,6 @@ public:
 		UseLegacyPaths		= 1 >> 1
 	};
 
-	
 	/**
 	 * Creates and initializes a new instance.
 	 *
@@ -49,7 +48,7 @@ public:
 
 	virtual FScreenshotExportResult ExportScreenshotComparisonResult(FString ScreenshotName, FString ExportPath = TEXT(""), bool bOnlyIncoming = false) override;
 
-	virtual bool OpenComparisonReports(FString ImportPath, TArray<FComparisonReport>& OutReports) override;
+	TFuture<TSharedPtr<TArray<FComparisonReport>>> OpenComparisonReportsAsync(const FString& ImportPath) override;
 	
 	virtual FString GetIdealApprovedFolderForImage(const FAutomationScreenshotMetadata& MetaData) const override;
 
@@ -82,6 +81,7 @@ private:
 
 	FString ScreenshotTempDeltaFolder;
 	FString ScreenshotResultsFolder;
+	TSharedPtr<TArray<FString>> PendingComparisonReportPaths;
 
 	TMap<FString, FString> FallbackPlatforms;
 
diff --git a/Engine/Source/Developer/ScreenShotComparisonTools/Public/Interfaces/IScreenShotManager.h b/Engine/Source/Developer/ScreenShotComparisonTools/Public/Interfaces/IScreenShotManager.h
--- a/Engine/Source/Developer/ScreenShotComparisonTools/Public/Interfaces/IScreenShotManager.h
+++ b/Engine/Source/Developer/ScreenShotComparisonTools/Public/Interfaces/IScreenShotManager.h
@@ -71,10 +71,10 @@ public:
 	virtual FScreenshotExportResult ExportScreenshotComparisonResult(FString ScreenshotName, FString ExportPath = TEXT(""), bool bOnlyIncoming = false) = 0;
 
 	/**
-	 * Imports screenshot comparison data from a given path.
+	 * Imports screenshot comparison data from a given path asynchronously.
 	 */
-	virtual bool OpenComparisonReports(FString ImportPath, TArray<FComparisonReport>& OutReports) = 0;
-	
+	virtual TFuture<TSharedPtr<TArray<FComparisonReport>>> OpenComparisonReportsAsync(const FString& ImportPath) = 0;
+
 	/**
 	* Calculate the ideal path for already already approved (ground truth) images. This handles abstracting away where Ground Truth images should be saved
 	* based on the platform/rhi/etc.
