Fix new screenshot not displaying when filter only new screenshot is enabled

#rnx

diff --git a/Engine/Source/Developer/ScreenShotComparison/Private/Widgets/SScreenShotBrowser.cpp b/Engine/Source/Developer/ScreenShotComparison/Private/Widgets/SScreenShotBrowser.cpp
--- a/Engine/Source/Developer/ScreenShotComparison/Private/Widgets/SScreenShotBrowser.cpp
+++ b/Engine/Source/Developer/ScreenShotComparison/Private/Widgets/SScreenShotBrowser.cpp
@@ -431,7 +431,7 @@ bool SScreenShotBrowser::MatchesReportFilterCriteria(const FString& ComparisonNa
 	}
 
 	bool bIsNew = ComparisonResult.IsNew();
-	bool bIsFail = !ComparisonResult.AreSimilar();
+	bool bIsFail = !bIsNew && !ComparisonResult.AreSimilar();
 	bool bIsPass = !bIsNew && !bIsFail;
 
 	if (bIsPass && !bDisplayingSuccess)
