fix a crash when deleteing screenshot comparison report of new ground truth

#rnx

diff --git a/Engine/Source/Developer/ScreenShotComparison/Private/Widgets/SScreenComparisonRow.cpp b/Engine/Source/Developer/ScreenShotComparison/Private/Widgets/SScreenComparisonRow.cpp
--- a/Engine/Source/Developer/ScreenShotComparison/Private/Widgets/SScreenComparisonRow.cpp
+++ b/Engine/Source/Developer/ScreenShotComparison/Private/Widgets/SScreenComparisonRow.cpp
@@ -152,6 +152,32 @@ private:
 	TSharedPtr<SCheckBox> DeltaCheckbox;
 };
 
+class SNewImage : public SCompoundWidget
+{
+public:
+	SLATE_BEGIN_ARGS(SNewImage) {}
+		SLATE_ARGUMENT(TSharedPtr<FSlateDynamicImageBrush>, NewBaseImage)
+	SLATE_END_ARGS()
+
+	void Construct(const FArguments& InArgs)
+	{
+		NewBaseImage = InArgs._NewBaseImage;
+
+		ChildSlot
+			[
+				SNew(SScaleBox)
+					.Stretch(EStretch::ScaleToFit)
+					[
+						SNew(SImage)
+							.Image(NewBaseImage.Get())
+					]
+			];
+	}
+
+
+private:
+	TSharedPtr<FSlateDynamicImageBrush> NewBaseImage;
+};
 
 void SScreenComparisonRow::Construct( const FArguments& InArgs, const TSharedRef<STableViewBase>& InOwnerTableView )
 {
@@ -608,12 +634,8 @@ FReply SScreenComparisonRow::OnCompareNewImage(const FGeometry& InGeometry, cons
 			.ActivationPolicy(EWindowActivationPolicy::Always)
 			.Content()
 			[
-				SNew(SScaleBox)
-				.Stretch(EStretch::ScaleToFit)
-				[
-					SNew(SImage)
-					.Image(UnapprovedImage.Get())
-				]
+				SNew(SNewImage)
+				.NewBaseImage(UnapprovedImage)
 			];
 
 		FSlateApplication::Get().AddWindowAsNativeChild(PopupWindow, ParentWindow, true);
